Copying //UE4/Dev-Enterprise to //UE4/Dev-Main (Source: //UE4/Dev-Enterprise @ 3972172)

diff --git a/Engine/Source/Runtime/Slate/Public/Widgets/SViewport.h b/Engine/Source/Runtime/Slate/Public/Widgets/SViewport.h
--- a/Engine/Source/Runtime/Slate/Public/Widgets/SViewport.h
+++ b/Engine/Source/Runtime/Slate/Public/Widgets/SViewport.h
@@ -170,6 +170,23 @@ public:
 		bRenderDirectlyToWindow = bInRenderDirectlyToWindow;
 	}
 
+	/**
+	 * If true, the viewport's texture alpha is ignored when performing blending.  In this case only the viewport tint opacity is used
+	 * If false, the texture alpha is used during blending
+	 * 
+	 * @param bIgnoreTextureAlpha If texture alpha should be ignored when blending.
+	 */
+	void SetIgnoreTextureAlpha(const bool bInIgnoreTextureAlpha)
+	{
+		bIgnoreTextureAlpha = bInIgnoreTextureAlpha;
+	}
+
+	/** @return Whether or not to ignore texture alpha when blending */
+	bool GetIgnoreTextureAlpha(void) const
+	{
+		return bIgnoreTextureAlpha;
+	}
+
 	/**
 	 * Sets whether stereo rendering is allowed for this viewport.  Advanced use only
 	 * 
diff --git a/Engine/Source/Runtime/SlateCore/Public/Styling/SlateTypes.h b/Engine/Source/Runtime/SlateCore/Public/Styling/SlateTypes.h
--- a/Engine/Source/Runtime/SlateCore/Public/Styling/SlateTypes.h
+++ b/Engine/Source/Runtime/SlateCore/Public/Styling/SlateTypes.h
@@ -30,6 +30,16 @@ enum class EConsumeMouseWheel : uint8
 	Never,
 };
 
+/** Used to determine which search method we should use when finding a suitable parent window */
+UENUM(BlueprintType)
+enum class ESlateParentWindowSearchMethod : uint8
+{
+	/** Favor using the active window (will fallback to the main window if the active window is unsuitable) */
+	ActiveWindow,
+
+	/** Favor using the main window */
+	MainWindow,
+};
 
 /** Type of check box */
 UENUM()
diff --git a/Engine/Source/Runtime/TimeManagement/Private/FixedFrameRateCustomTimeStep.cpp b/Engine/Source/Runtime/TimeManagement/Private/FixedFrameRateCustomTimeStep.cpp
new file mode 100644
--- /dev/null
+++ b/Engine/Source/Runtime/TimeManagement/Private/FixedFrameRateCustomTimeStep.cpp
@@ -0,0 +1,12 @@
+// Copyright 1998-2018 Epic Games, Inc. All Rights Reserved.
+
+#include "FixedFrameRateCustomTimeStep.h"
+
+
+UFixedFrameRateCustomTimeStep::UFixedFrameRateCustomTimeStep(const FObjectInitializer& ObjectInitializer)
+	: Super(ObjectInitializer)
+	, FixedFrameRate(30, 1)
+{
+}
+
+
diff --git a/Engine/Source/Runtime/TimeManagement/Private/ITimeSynchronizationManager.cpp b/Engine/Source/Runtime/TimeManagement/Private/ITimeSynchronizationManager.cpp
new file mode 100644
--- /dev/null
+++ b/Engine/Source/Runtime/TimeManagement/Private/ITimeSynchronizationManager.cpp
@@ -0,0 +1,35 @@
+// Copyright 1998-2018 Epic Games, Inc. All Rights Reserved.
+
+#include "ITimeSynchronizationManager.h"
+
+#include "ITimeManagementModule.h"
+
+
+FTimecode UTimeSynchronizationManagerHelpers::GetCurrentTimecode()
+{
+	if (const ITimeSynchronizationManager* Manager = ITimeManagementModule::Get().GetTimeSynchronizationManager())
+	{
+		return Manager->GetCurrentTimecode();
+	}
+	return FTimecode();
+}
+
+bool UTimeSynchronizationManagerHelpers::IsSynchronizing()
+{
+	bool bResult = false;
+	if (const ITimeSynchronizationManager* Manager = ITimeManagementModule::Get().GetTimeSynchronizationManager())
+	{
+		bResult = Manager->IsSynchronizing();
+	}
+	return bResult;
+}
+
+bool UTimeSynchronizationManagerHelpers::IsSynchronized()
+{
+	bool bResult = false;
+	if (const ITimeSynchronizationManager* Manager = ITimeManagementModule::Get().GetTimeSynchronizationManager())
+	{
+		bResult = Manager->IsSynchronized();
+	}
+	return bResult;
+}
diff --git a/Engine/Source/Runtime/TimeManagement/Private/TimeManagementModule.cpp b/Engine/Source/Runtime/TimeManagement/Private/TimeManagementModule.cpp
--- a/Engine/Source/Runtime/TimeManagement/Private/TimeManagementModule.cpp
+++ b/Engine/Source/Runtime/TimeManagement/Private/TimeManagementModule.cpp
@@ -1,5 +1,11 @@
 // Copyright 1998-2018 Epic Games, Inc. All Rights Reserved.
 
-#include "Modules/ModuleManager.h"
+#include "TimeManagementModule.h"
 
-IMPLEMENT_MODULE(FDefaultModuleImpl, TimeManagement);
\ No newline at end of file
+
+FTimeManagementModule::FTimeManagementModule()
+	: TimeSyncManager(nullptr)
+{
+}
+
+IMPLEMENT_MODULE(FTimeManagementModule, TimeManagement);
\ No newline at end of file
diff --git a/Engine/Source/Runtime/TimeManagement/Private/TimeManagementModule.h b/Engine/Source/Runtime/TimeManagement/Private/TimeManagementModule.h
new file mode 100644
--- /dev/null
+++ b/Engine/Source/Runtime/TimeManagement/Private/TimeManagementModule.h
@@ -0,0 +1,22 @@
+// Copyright 1998-2018 Epic Games, Inc. All Rights Reserved.
+
+#pragma once
+
+#include "ITimeManagementModule.h"
+
+class FTimeManagementModule : public ITimeManagementModule
+{
+public:
+	FTimeManagementModule();
+
+public:
+	//~ Begin ITimeManagementModule Interface
+	virtual ITimeSynchronizationManager* GetTimeSynchronizationManager() override { return TimeSyncManager; }
+	virtual ITimeSynchronizationManager* GetTimeSynchronizationManager() const override { return TimeSyncManager; }
+
+	virtual void SetTimeSynchronizationManager(ITimeSynchronizationManager* Manager) override { TimeSyncManager = Manager; }
+	//~ End ITimeManagementModule Interface
+
+private:
+	ITimeSynchronizationManager* TimeSyncManager;
+};
diff --git a/Engine/Source/Runtime/TimeManagement/Public/FixedFrameRateCustomTimeStep.h b/Engine/Source/Runtime/TimeManagement/Public/FixedFrameRateCustomTimeStep.h
new file mode 100644
--- /dev/null
+++ b/Engine/Source/Runtime/TimeManagement/Public/FixedFrameRateCustomTimeStep.h
@@ -0,0 +1,26 @@
+// Copyright 1998-2018 Epic Games, Inc. All Rights Reserved.
+
+#pragma once
+
+#include "Engine/EngineCustomTimeStep.h"
+
+#include "FrameRate.h"
+
+#include "FixedFrameRateCustomTimeStep.generated.h"
+
+
+
+/**
+ * Class to control the Engine TimeStep via a FixedFrameRate
+ */
+UCLASS(Abstract)
+class TIMEMANAGEMENT_API UFixedFrameRateCustomTimeStep : public UEngineCustomTimeStep
+{
+	GENERATED_UCLASS_BODY()
+
+public:
+
+	/** The fixed FrameRate */
+	UPROPERTY(EditAnywhere, Category=Time)
+	FFrameRate FixedFrameRate;
+};
diff --git a/Engine/Source/Runtime/TimeManagement/Public/ITimeManagementModule.h b/Engine/Source/Runtime/TimeManagement/Public/ITimeManagementModule.h
new file mode 100644
--- /dev/null
+++ b/Engine/Source/Runtime/TimeManagement/Public/ITimeManagementModule.h
@@ -0,0 +1,42 @@
+// Copyright 1998-2018 Epic Games, Inc. All Rights Reserved.
+
+#pragma once
+
+#include "CoreMinimal.h"
+#include "Modules/ModuleInterface.h"
+#include "Modules/ModuleManager.h"
+
+
+#define TIMEMANAGEMENT_MODULE_NAME TEXT("TimeManagement")
+
+class ITimeSynchronizationManager;
+
+class TIMEMANAGEMENT_API ITimeManagementModule : public IModuleInterface
+{
+public:
+	
+	/**
+	 * Singleton-like access to ITimeManagementModule
+	 * @return Returns ITimeManagementModule singleton instance, loading the module on demand if needed
+	 */
+	static inline ITimeManagementModule& Get()
+	{
+		return FModuleManager::LoadModuleChecked<ITimeManagementModule>(TIMEMANAGEMENT_MODULE_NAME);
+	}
+
+	/**
+	 * Checks to see if this module is loaded and ready.  It is only valid to call Get() if IsAvailable() returns true.
+	 * @return True if the module is loaded and ready to use
+	 */
+	static inline bool IsAvailable()
+	{
+		return FModuleManager::Get().IsModuleLoaded(TIMEMANAGEMENT_MODULE_NAME);
+	}
+
+	/** Return the current TimeSynchronizationManager that is responsible to own the current timecode for the Engine */
+	virtual ITimeSynchronizationManager* GetTimeSynchronizationManager() = 0;
+	virtual ITimeSynchronizationManager* GetTimeSynchronizationManager() const = 0;
+	
+	/** Set the current TimeSynchronizationManager that is responsible to own the current timecode for the Engine */
+	virtual void SetTimeSynchronizationManager(ITimeSynchronizationManager* Manager) = 0;
+};
diff --git a/Engine/Source/Runtime/TimeManagement/Public/ITimeSynchronizationManager.h b/Engine/Source/Runtime/TimeManagement/Public/ITimeSynchronizationManager.h
new file mode 100644
--- /dev/null
+++ b/Engine/Source/Runtime/TimeManagement/Public/ITimeSynchronizationManager.h
@@ -0,0 +1,64 @@
+// Copyright 1998-2018 Epic Games, Inc. All Rights Reserved.
+
+#pragma once
+
+#include "FrameRate.h"
+#include "Timecode.h"
+
+#include "ITimeSynchronizationManager.generated.h"
+
+
+/**
+ * Interface for Time Synchronization Manager module.
+ */
+class ITimeSynchronizationManager
+{
+public:
+
+	/**
+	 * Get the manager's actual Timecode for this running frame.
+	 * Only valid when IsSynchronized() is true
+	 */
+	UFUNCTION(BlueprintCallable, Category = "TimeSynchronization")
+	virtual FTimecode GetCurrentTimecode() const = 0;
+	
+	/** Get the frame rate of the manager in frames per second */
+	virtual FFrameRate GetFrameRate() const = 0;
+
+	/**
+	 * All sources are currently synchronizing/buffering, but not all sources are ready yet.
+	 * GetCurrentTimecode is not valid in that state.
+	 */
+	UFUNCTION(BlueprintCallable, Category = "TimeSynchronization")
+	virtual bool IsSynchronizing() const = 0;
+
+	/** All sources are currently synchronizing/buffering, but not all sources are ready yet. */
+	UFUNCTION(BlueprintCallable, Category = "TimeSynchronization")
+	virtual bool IsSynchronized() const = 0;
+};
+
+
+UCLASS(Transient, MinimalAPI, meta = (ScriptName = "TimeSynchronizationLibrary"))
+class UTimeSynchronizationManagerHelpers : public UObject
+{
+	GENERATED_BODY()
+
+public:
+	/**
+	 * Get the manager's actual Timecode for this running frame.
+	 * Only valid when IsSynchronized() is true
+	 */
+	UFUNCTION(BlueprintPure, Category = "TimeSynchronization")
+	static FTimecode GetCurrentTimecode();
+
+	/**
+	 * All sources are currently synchronizing/buffering, but not all sources are ready yet.
+	 * GetCurrentTimecode is not valid in that state.
+	 */
+	UFUNCTION(BlueprintPure, Category = "TimeSynchronization")
+	static bool IsSynchronizing();
+
+	/** All sources are currently synchronizing/buffering, but not all sources are ready yet. */
+	UFUNCTION(BlueprintPure, Category = "TimeSynchronization")
+	static bool IsSynchronized();
+};
diff --git a/Engine/Source/Runtime/TimeManagement/Public/TimeSynchronizationSource.h b/Engine/Source/Runtime/TimeManagement/Public/TimeSynchronizationSource.h
new file mode 100644
--- /dev/null
+++ b/Engine/Source/Runtime/TimeManagement/Public/TimeSynchronizationSource.h
@@ -0,0 +1,58 @@
+// Copyright 1998-2018 Epic Games, Inc. All Rights Reserved.
+
+#pragma once
+
+#include "UObject/ObjectMacros.h"
+#include "UObject/Object.h"
+
+#include "FrameRate.h"
+#include "FrameTime.h"
+
+#include "TimeSynchronizationSource.generated.h"
+
+/**
+* Base class for sources to be used for time synchronization
+*/
+UCLASS(Abstract)
+class TIMEMANAGEMENT_API UTimeSynchronizationSource : public UObject
+{
+	GENERATED_BODY()
+
+public:
+	/** If the source has a time base that can be used to synchronize other sources. */
+	UPROPERTY(EditAnywhere, Category = Synchronization)
+	bool bUseForSynchronization;
+
+	/** Extra frame to buffered before allowing the Manager to by in synchronized mode. */
+	UPROPERTY(EditAnywhere, Category = Synchronization, Meta = (ClampMin = 0, ClampMax = 60, EditCondition = "bUseForSynchronization"))
+	int32 NumberOfExtraBufferedFrame;
+
+	/** Fixed delay in seconds to align with other sources. */
+	UPROPERTY(EditAnywhere, Category = Synchronization, Meta = (EditCondition = "!bUseForSynchronization"))
+	float TimeDelay;
+
+public:
+	/** Get next available sample Time based on the source FrameRate */
+	virtual FFrameTime GetNextSampleTime() const PURE_VIRTUAL(UTimeSynchronizationSource::GetNextSampleTime, return FFrameTime();)
+
+	/** Get number of available samples buffered in the source */
+	virtual int32 GetAvailableSampleCount() const PURE_VIRTUAL(UTimeSynchronizationSource::GetAvailableSampleCount, return 0;)
+
+	/** Get the source actual FrameRate */
+	virtual FFrameRate GetFrameRate() const PURE_VIRTUAL(UTimeSynchronizationSource::GetFrameRate, return FFrameRate();)
+
+	/** Used to know if the source is ready to be used for synchronization. */
+	virtual bool IsReady() const PURE_VIRTUAL(UTimeSynchronizationSource::IsReady, return false;)
+
+	/** Open the source to initiate frame acquisition */
+	virtual bool Open() PURE_VIRTUAL(UTimeSynchronizationSource::Open, return false;)
+
+	/** Start Rolling/Playing the source */
+	virtual void Start() PURE_VIRTUAL(UTimeSynchronizationSource::Start, return;)
+
+	/** Stop the source from rolling. The source may clear all his buffered frame. */
+	virtual void Close() PURE_VIRTUAL(UTimeSynchronizationSource::Close, return;)
+
+	/** Name to used when displaying an error message or to used in UI. */
+	virtual FString GetDisplayName() const PURE_VIRTUAL(UTimeSynchronizationSource::GetDisplayName, return FString(););
+};
diff --git a/Engine/Source/Runtime/Windows/D3D11RHI/Private/D3D11Viewport.cpp b/Engine/Source/Runtime/Windows/D3D11RHI/Private/D3D11Viewport.cpp
--- a/Engine/Source/Runtime/Windows/D3D11RHI/Private/D3D11Viewport.cpp
+++ b/Engine/Source/Runtime/Windows/D3D11RHI/Private/D3D11Viewport.cpp
@@ -6,6 +6,7 @@
 
 #include "D3D11RHIPrivate.h"
 #include "RenderCore.h"
+#include "Engine/RendererSettings.h"
 
 #ifndef D3D11_WITH_DWMAPI
 #if WINVER > 0x502		// Windows XP doesn't support DWM
@@ -490,7 +491,8 @@ FViewportRHIRef FD3D11DynamicRHI::RHICreateViewport(void* WindowHandle,uint32 Si
 	// Use a default pixel format if none was specified	
 	if (PreferredPixelFormat == EPixelFormat::PF_Unknown)
 	{
-		PreferredPixelFormat = EPixelFormat::PF_R8G8B8A8;
+		static const auto CVarDefaultBackBufferPixelFormat = IConsoleManager::Get().FindTConsoleVariableDataInt(TEXT("r.DefaultBackBufferPixelFormat"));
+		PreferredPixelFormat = EDefaultBackBufferPixelFormat::Convert2PixelFormat(CVarDefaultBackBufferPixelFormat->GetValueOnGameThread());
 	}
 
 	return new FD3D11Viewport(this,(HWND)WindowHandle,SizeX,SizeY,bIsFullscreen,PreferredPixelFormat);
@@ -506,16 +508,16 @@ void FD3D11DynamicRHI::RHIResizeViewport(FViewportRHIParamRef ViewportRHI,uint32
 
 void FD3D11DynamicRHI::RHIResizeViewport(FViewportRHIParamRef ViewportRHI, uint32 SizeX, uint32 SizeY, bool bIsFullscreen, EPixelFormat PreferredPixelFormat)
 {
-	FD3D11Viewport* Viewport = ResourceCast(ViewportRHI);
-
 	check(IsInGameThread());
 
 	// Use a default pixel format if none was specified	
 	if (PreferredPixelFormat == EPixelFormat::PF_Unknown)
 	{
-		PreferredPixelFormat = EPixelFormat::PF_R8G8B8A8;
+		static const auto CVarDefaultBackBufferPixelFormat = IConsoleManager::Get().FindTConsoleVariableDataInt(TEXT("r.DefaultBackBufferPixelFormat"));
+		PreferredPixelFormat = EDefaultBackBufferPixelFormat::Convert2PixelFormat(CVarDefaultBackBufferPixelFormat->GetValueOnGameThread());
 	}
 
+	FD3D11Viewport* Viewport = ResourceCast(ViewportRHI);
 	Viewport->Resize(SizeX, SizeY, bIsFullscreen, PreferredPixelFormat);
 }
 
