Copying //UE4/Dev-Sequencer to //UE4/Dev-Main (Source: //UE4/Dev-Sequencer @ 4044339)

diff --git a/Engine/Source/Runtime/TimeManagement/Private/FrameNumberTimeEvaluator.cpp b/Engine/Source/Runtime/TimeManagement/Private/FrameNumberTimeEvaluator.cpp
--- a/Engine/Source/Runtime/TimeManagement/Private/FrameNumberTimeEvaluator.cpp
+++ b/Engine/Source/Runtime/TimeManagement/Private/FrameNumberTimeEvaluator.cpp
@@ -60,8 +60,11 @@ FFrameNumberTimeEvaluator::FFrameNumberTimeEvaluator()
 	TimeJumpTable.MapBinary<FStar>([](double A, double B) { return A * B; });
 }
 
-TValueOrError<FFrameTime, FExpressionError> FFrameNumberTimeEvaluator::EvaluateTimecode(const TCHAR* InExpression, const FFrameRate& InDisplayFrameRate, const FFrameRate& InFrameResolution) const
+TValueOrError<FFrameTime, FExpressionError> FFrameNumberTimeEvaluator::EvaluateTimecode(const TCHAR* InExpression, const FFrameRate& InDisplayFrameRate, const FFrameRate& InTickResolution, bool& OutDirectlyParsed) const
 {
+	OutDirectlyParsed = false;
+
+
 	using namespace ExpressionParser;
 
 	TValueOrError<TArray<FExpressionToken>, FExpressionError> LexResult = ExpressionParser::Lex(InExpression, TimecodeTokenDefinitions);
@@ -83,6 +86,11 @@ TValueOrError<FFrameTime, FExpressionError> FFrameNumberTimeEvaluator::EvaluateT
 		if (Node.Cast<FDropcodeDelimiter>())
 		{
 			bIsDropcode = true;
+			OutDirectlyParsed = true;
+		}
+		else if (Node.Cast<FTimecodeDelimiter>())
+		{
+			OutDirectlyParsed = true;
 		}
 		else if (Node.Cast<FBracketStart>() || Node.Cast<FBracketEnd>() || Node.Cast<FPlus>() || Node.Cast<FMinus>())
 		{
@@ -100,6 +108,7 @@ TValueOrError<FFrameTime, FExpressionError> FFrameNumberTimeEvaluator::EvaluateT
 	// There definitely can't be more than 7 tokens, but we do accept less tokens
 	if (Tokens.Num() > 7)
 	{
+		OutDirectlyParsed = false;
 		return MakeError(NSLOCTEXT("TimeManagement", "UnrecognizedTimecode", "Format not recognized as Timecode"));
 	}
 
@@ -119,6 +128,7 @@ TValueOrError<FFrameTime, FExpressionError> FFrameNumberTimeEvaluator::EvaluateT
 			// Every other one should be a numeric, so if it's not, we're not sure what format it is.
 			if (!Node.Cast<double>())
 			{
+				OutDirectlyParsed = false;
 				return MakeError(NSLOCTEXT("TimeManagement", "UnrecognizedTimecode", "Format not recognized as Timecode"));
 			}
 
@@ -135,6 +145,7 @@ TValueOrError<FFrameTime, FExpressionError> FFrameNumberTimeEvaluator::EvaluateT
 			// Every other one should be a delimiter, if it's not, we're not sure what they're putting in.
 			if (!(Node.Cast<FTimecodeDelimiter>() || Node.Cast<FDropcodeDelimiter>()))
 			{
+				OutDirectlyParsed = false;
 				return MakeError(NSLOCTEXT("TimeManagement", "UnrecognizedTimecode", "Format not recognized as Timecode"));
 			}
 		}
@@ -171,7 +182,7 @@ TValueOrError<FFrameTime, FExpressionError> FFrameNumberTimeEvaluator::EvaluateT
 	bool bDropFrameSupported = FTimecode::IsDropFormatTimecodeSupported(InDisplayFrameRate);
 	FTimecode Timecode(Hours, Minutes, Seconds, Frames, bIsDropcode && bDropFrameSupported);
 
-	FFrameNumber TotalFrames = FFrameRate::TransformTime(FFrameTime(Timecode.ToFrameNumber(InDisplayFrameRate)), InDisplayFrameRate, InFrameResolution).RoundToFrame();
+	FFrameNumber TotalFrames = FFrameRate::TransformTime(FFrameTime(Timecode.ToFrameNumber(InDisplayFrameRate)), InDisplayFrameRate, InTickResolution).RoundToFrame();
 	if (bIsNegative)
 	{
 		TotalFrames = -TotalFrames;
@@ -180,7 +191,7 @@ TValueOrError<FFrameTime, FExpressionError> FFrameNumberTimeEvaluator::EvaluateT
 	return MakeValue(TotalFrames);
 }
 
-TValueOrError<FFrameTime, FExpressionError> FFrameNumberTimeEvaluator::EvaluateFrame(const TCHAR* InExpression, const FFrameRate& InDisplayFrameRate, const FFrameRate& InFrameResolution, bool& OutDirectlyParsed) const
+TValueOrError<FFrameTime, FExpressionError> FFrameNumberTimeEvaluator::EvaluateFrame(const TCHAR* InExpression, const FFrameRate& InDisplayFrameRate, const FFrameRate& InTickResolution, bool& OutDirectlyParsed) const
 {
 	using namespace ExpressionParser;
 	OutDirectlyParsed = false;
@@ -217,6 +228,7 @@ TValueOrError<FFrameTime, FExpressionError> FFrameNumberTimeEvaluator::EvaluateF
 	// If they're jumping to a specific frame so there should only be one token.
 	if (Tokens.Num() != 1)
 	{
+		OutDirectlyParsed = false;
 		return MakeError(NSLOCTEXT("TimeManagement", "UnrecognizedFrame", "Format not recognized as a Frame number"));
 	}
 
@@ -226,7 +238,7 @@ TValueOrError<FFrameTime, FExpressionError> FFrameNumberTimeEvaluator::EvaluateF
 		Frame = -Frame;
 	}
 
-	FFrameTime Result = FFrameRate::TransformTime(Frame, InDisplayFrameRate, InFrameResolution);
+	FFrameTime Result = FFrameRate::TransformTime(Frame, InDisplayFrameRate, InTickResolution);
 	return MakeValue(Result);
 }
 
@@ -277,7 +289,6 @@ TValueOrError<FFrameTime, FExpressionError> FFrameNumberTimeEvaluator::EvaluateT
 		{
 			FStringToken Context = Tokens[i].Context;
 			const FExpressionNode& Node = Tokens[i].Node;
-			UE_LOG(LogTemp, Log, TEXT("i: %d Token: %s"), i, *Tokens[i].Context.GetString());
 
 			if (!(Node.Cast<FHour>() || Node.Cast<FMinute>() || Node.Cast<FSecond>() || Node.Cast<FMillisecond>()))
 			{
@@ -354,12 +365,14 @@ TValueOrError<FFrameTime, FExpressionError> FFrameNumberTimeEvaluator::EvaluateT
 		}
 
 		// There was only one token (or no valid tokens) and it wasn't a number, therefor we don't know what it is.
+		OutDirectlyParsed = false;
 		return MakeError(NSLOCTEXT("TimeManagement", "UnrecognizedTime", "Format not recognized as a time"));
 	}
 
 	TValueOrError<TArray<FCompiledToken>, FExpressionError> CompilationResult = ExpressionParser::Compile(MoveTemp(Tokens), TimeGrammar);
 	if (!CompilationResult.IsValid())
 	{
+		OutDirectlyParsed = false;
 		return MakeError(CompilationResult.StealError());
 	}
 
@@ -367,6 +380,7 @@ TValueOrError<FFrameTime, FExpressionError> FFrameNumberTimeEvaluator::EvaluateT
 	TValueOrError<FExpressionNode, FExpressionError> MillisecondEvaluationResult = ExpressionParser::Evaluate(CompilationResult.GetValue(), Env);
 	if (!MillisecondEvaluationResult.IsValid())
 	{
+		OutDirectlyParsed = false;
 		return MakeError(MillisecondEvaluationResult.GetError());
 	}
 
@@ -380,6 +394,7 @@ TValueOrError<FFrameTime, FExpressionError> FFrameNumberTimeEvaluator::EvaluateT
 		return MakeValue(Result);
 	}
 
+	OutDirectlyParsed = false;
 	return MakeError(NSLOCTEXT("TimeManagement", "UnrecognizedTimeResult", "Unrecognized result returned from expression"));
 
 }
\ No newline at end of file
diff --git a/Engine/Source/Runtime/TimeManagement/Private/FrameRate.cpp b/Engine/Source/Runtime/TimeManagement/Private/FrameRate.cpp
--- a/Engine/Source/Runtime/TimeManagement/Private/FrameRate.cpp
+++ b/Engine/Source/Runtime/TimeManagement/Private/FrameRate.cpp
@@ -21,6 +21,9 @@ DEFINE_EXPRESSION_NODE_TYPE(FFrameRate, 0x4EDAA92F, 0xB75E4B9E, 0xB7E0ABC2, 0x8D
 DEFINE_EXPRESSION_NODE_TYPE(FSeconds, 0x3DC5F60D, 0x934E4753, 0xA80CD6D0, 0xE9EB4640)
 DEFINE_EXPRESSION_NODE_TYPE(FFramesPerSecond, 0x8423B4AE, 0x2FF64795, 0xA7EFFAC0, 0xC560531A)
 
+const float FFrameTime::MaxSubframe = 0.99999994f;
+
+
 /** A basic math expression evaluator */
 class FFrameRateParser
 {
diff --git a/Engine/Source/Runtime/TimeManagement/Private/TimeManagementTests.cpp b/Engine/Source/Runtime/TimeManagement/Private/TimeManagementTests.cpp
--- a/Engine/Source/Runtime/TimeManagement/Private/TimeManagementTests.cpp
+++ b/Engine/Source/Runtime/TimeManagement/Private/TimeManagementTests.cpp
@@ -372,11 +372,11 @@ bool FFrameTimeConversionTest::RunTest(const FString& Parameters)
 		FFrameTime ExpectedTimes[] = {
 			FFrameTime(-5, 0.f),  FFrameTime(-5, 0.05f), FFrameTime(-5, 0.16f), FFrameTime(-5, 0.32f), FFrameTime(-5, 0.499999985f),
 			FFrameTime(-3, 0.f),  FFrameTime(-3, 0.05f), FFrameTime(-3, 0.16f), FFrameTime(-3, 0.32f), FFrameTime(-3, 0.499999985f),
-			FFrameTime(-3, 0.5f), FFrameTime(-3, 0.55f), FFrameTime(-3, 0.66f), FFrameTime(-3, 0.82f), FFrameTime(-3, 0.999999940f),
-			FFrameTime(-2, 0.5f), FFrameTime(-2, 0.55f), FFrameTime(-2, 0.66f), FFrameTime(-2, 0.82f), FFrameTime(-2, 0.999999940f),
+			FFrameTime(-3, 0.5f), FFrameTime(-3, 0.55f), FFrameTime(-3, 0.66f), FFrameTime(-3, 0.82f), FFrameTime(-3, FFrameTime::MaxSubframe),
+			FFrameTime(-2, 0.5f), FFrameTime(-2, 0.55f), FFrameTime(-2, 0.66f), FFrameTime(-2, 0.82f), FFrameTime(-2, FFrameTime::MaxSubframe),
 			FFrameTime(0,  0.f),  FFrameTime(0,  0.05f), FFrameTime(0,  0.16f), FFrameTime(0,  0.32f), FFrameTime(0,  0.499999985f),
-			FFrameTime(1,  0.5f), FFrameTime(1,  0.55f), FFrameTime(1,  0.66f), FFrameTime(1,  0.82f), FFrameTime(1,  0.999999940f),
-			FFrameTime(2,  0.5f), FFrameTime(2,  0.55f), FFrameTime(2,  0.66f), FFrameTime(2,  0.82f), FFrameTime(2,  0.999999940f),
+			FFrameTime(1,  0.5f), FFrameTime(1,  0.55f), FFrameTime(1,  0.66f), FFrameTime(1,  0.82f), FFrameTime(1,  FFrameTime::MaxSubframe),
+			FFrameTime(2,  0.5f), FFrameTime(2,  0.55f), FFrameTime(2,  0.66f), FFrameTime(2,  0.82f), FFrameTime(2,  FFrameTime::MaxSubframe),
 			FFrameTime(3,  0.f),  FFrameTime(3,  0.05f), FFrameTime(3,  0.16f), FFrameTime(3,  0.32f), FFrameTime(3,  0.499999985f),
 			FFrameTime(5,  0.f),  FFrameTime(5,  0.05f), FFrameTime(5,  0.16f), FFrameTime(5,  0.32f), FFrameTime(5,  0.499999985f),
 		};
diff --git a/Engine/Source/Runtime/TimeManagement/Public/FrameNumberNumericInterface.h b/Engine/Source/Runtime/TimeManagement/Public/FrameNumberNumericInterface.h
--- a/Engine/Source/Runtime/TimeManagement/Public/FrameNumberNumericInterface.h
+++ b/Engine/Source/Runtime/TimeManagement/Public/FrameNumberNumericInterface.h
@@ -31,20 +31,20 @@ DECLARE_DELEGATE_RetVal(FFrameRate, FOnGetFrameRate)
 */
 struct FFrameNumberInterface : public INumericTypeInterface<double>
 {
-	FFrameNumberInterface(FOnGetDisplayFormat InGetDisplayFormat, FOnGetZeroPad InOnGetZeroPadFrameNumber, FOnGetFrameRate InGetFrameResolution, FOnGetFrameRate InGetPlayRate)
+	FFrameNumberInterface(FOnGetDisplayFormat InGetDisplayFormat, FOnGetZeroPad InOnGetZeroPadFrameNumber, FOnGetFrameRate InGetTickResolution, FOnGetFrameRate InGetPlayRate)
 		: GetDisplayFormat(InGetDisplayFormat)
-		, GetFrameResolution(InGetFrameResolution)
+		, GetTickResolution(InGetTickResolution)
 		, GetPlayRate(InGetPlayRate)
 		, GetZeroPadFrames(InOnGetZeroPadFrameNumber)
 	{
 		check(InGetDisplayFormat.IsBound());
-		check(InGetFrameResolution.IsBound());
+		check(InGetTickResolution.IsBound());
 		check(InGetPlayRate.IsBound());
 	}
 
 private:
 	FOnGetDisplayFormat GetDisplayFormat;
-	FOnGetFrameRate GetFrameResolution;
+	FOnGetFrameRate GetTickResolution;
 	FOnGetFrameRate GetPlayRate;
 	FOnGetZeroPad GetZeroPadFrames;
 
@@ -65,7 +65,7 @@ private:
 
 	virtual FString ToString(const double& Value) const override
 	{
-		FFrameRate SourceFrameRate = GetFrameResolution.Execute();
+		FFrameRate SourceFrameRate = GetTickResolution.Execute();
 		FFrameRate DestinationFrameRate = GetPlayRate.Execute();
 		EFrameNumberDisplayFormats Format = GetDisplayFormat.Execute();
 
@@ -115,7 +115,7 @@ private:
 	virtual TOptional<double> FromString(const FString& InString, const double& InExistingValue) override
 	{
 		FFrameRate SourceFrameRate = GetPlayRate.Execute();
-		FFrameRate DestinationFrameRate = GetFrameResolution.Execute();
+		FFrameRate DestinationFrameRate = GetTickResolution.Execute();
 		EFrameNumberDisplayFormats FallbackFormat = GetDisplayFormat.Execute();
 
 		// We allow input in any format (time, frames or timecode) and we just convert it into into the internal sequence resolution.
@@ -124,20 +124,26 @@ private:
 
 		// All of these will convert into the frame resolution from the user's input before returning.
 		FFrameNumberTimeEvaluator Eval;
-		TValueOrError<FFrameTime, FExpressionError> TimecodeResult = Eval.EvaluateTimecode(*InString, SourceFrameRate, DestinationFrameRate);
+		bool bWasTimecodeText;
+		TValueOrError<FFrameTime, FExpressionError> TimecodeResult = Eval.EvaluateTimecode(*InString, SourceFrameRate, DestinationFrameRate, /*Out*/bWasTimecodeText);
 		bool bWasFrameText;
 		TValueOrError<FFrameTime, FExpressionError> FrameResult = Eval.EvaluateFrame(*InString, SourceFrameRate, DestinationFrameRate, /*Out*/ bWasFrameText);
 		bool bWasTimeText;
 		TValueOrError<FFrameTime, FExpressionError> TimeResult = Eval.EvaluateTime(*InString, DestinationFrameRate, /*Out*/ bWasTimeText);
 
-		if (TimecodeResult.IsValid())
-		{
-			return TOptional<double>(TimecodeResult.GetValue().GetFrame().Value);
-		}
-		else if (!bWasFrameText && !bWasTimeText)
+
+		// All three formats support ambiguous conversion where the user can enter "5" and wants it in the logical unit based on the current
+		// display format. This means 5 -> 5f, 5 ->5s, and 5 -> 5 frames (in timecode). We also support specifically specifying in a different
+		// format than your current display, ie if your display is in frames and you enter 1s, you get 1s in frames or 30 (at 30fps play rate).
+		if (!bWasTimecodeText && !bWasFrameText && !bWasTimeText)
 		{
-			// If the value was ambiguous ("5") and not specifically a frame ("5f") or a time "(5s") then we defer to the current display format.
-			if (TimeResult.IsValid() && FallbackFormat == EFrameNumberDisplayFormats::Seconds)
+			// They've entered an ambiguous number, so we'll check the display format and see if we did successfully parse as that type. If it
+			// was able to parse the input for the given display format, we return that.
+			if (TimecodeResult.IsValid() && (FallbackFormat == EFrameNumberDisplayFormats::DropFrameTimecode || FallbackFormat == EFrameNumberDisplayFormats::NonDropFrameTimecode))
+			{
+				return TOptional<double>(TimecodeResult.GetValue().GetFrame().Value);
+			}
+			else if (TimeResult.IsValid() && FallbackFormat == EFrameNumberDisplayFormats::Seconds)
 			{
 				return TOptional<double>(TimeResult.GetValue().GetFrame().Value);
 			}
@@ -145,20 +151,22 @@ private:
 			{
 				return TOptional<double>(FrameResult.GetValue().GetFrame().Value);
 			}
-			else
-			{ 
-				// We don't support ambiguous conversion to Timecode
-				return TOptional<double>();
-			}
+
+			// Whatever they entered wasn't understood by any of our parsers, so it was probably malformed or had letters, etc.
+			return TOptional<double>();
+		}
+
+		// If we've gotten here then they did explicitly specify a timecode so we return that.
+		if (bWasTimecodeText)
+		{
+			return TOptional<double>(TimecodeResult.GetValue().GetFrame().Value);
 		}
-		else if (TimeResult.IsValid() && !bWasFrameText)
+		else if (bWasTimeText)
 		{
-			// Both time and frames can parse "5" successfully, so if we know it's not specifically frame text
 			return TOptional<double>(TimeResult.GetValue().GetFrame().Value);
 		}
-		else if (FrameResult.IsValid() && !bWasTimeText)
+		else if (bWasFrameText)
 		{
-			// Both time and frames can parse "5" successfully and we know it's not specifically time
 			return TOptional<double>(FrameResult.GetValue().GetFrame().Value);
 		}
 
diff --git a/Engine/Source/Runtime/TimeManagement/Public/FrameNumberTimeEvaluator.h b/Engine/Source/Runtime/TimeManagement/Public/FrameNumberTimeEvaluator.h
--- a/Engine/Source/Runtime/TimeManagement/Public/FrameNumberTimeEvaluator.h
+++ b/Engine/Source/Runtime/TimeManagement/Public/FrameNumberTimeEvaluator.h
@@ -34,8 +34,8 @@ struct TIMEMANAGEMENT_API FFrameNumberTimeEvaluator
 	FFrameNumberTimeEvaluator();
 
 	/** Evaluate the given expression, resulting in either a double value, or an error */
-	TValueOrError<FFrameTime, FExpressionError> EvaluateTimecode(const TCHAR* InExpression, const FFrameRate& InDisplayFrameRate, const FFrameRate& InFrameResolution) const;
-	TValueOrError<FFrameTime, FExpressionError> EvaluateFrame(const TCHAR* InExpression, const FFrameRate& InDisplayFrameRate, const FFrameRate& InFrameResolution, bool& OutDirectlyParsed) const;
+	TValueOrError<FFrameTime, FExpressionError> EvaluateTimecode(const TCHAR* InExpression, const FFrameRate& InDisplayFrameRate, const FFrameRate& InTickResolution, bool& OutDirectlyParsed) const;
+	TValueOrError<FFrameTime, FExpressionError> EvaluateFrame(const TCHAR* InExpression, const FFrameRate& InDisplayFrameRate, const FFrameRate& InTickResolution, bool& OutDirectlyParsed) const;
 	TValueOrError<FFrameTime, FExpressionError> EvaluateTime(const TCHAR* InExpression, FFrameRate InFrameRate, bool& OutDirectlyParsed) const;
 
 	FTokenDefinitions TimecodeTokenDefinitions;
diff --git a/Engine/Source/Runtime/TimeManagement/Public/FrameRate.h b/Engine/Source/Runtime/TimeManagement/Public/FrameRate.h
--- a/Engine/Source/Runtime/TimeManagement/Public/FrameRate.h
+++ b/Engine/Source/Runtime/TimeManagement/Public/FrameRate.h
@@ -184,7 +184,7 @@ inline FFrameTime FFrameRate::AsFrameTime(double TimeInSeconds) const
 	double SubFrame = TimeAsFrame - FMath::FloorToDouble(TimeAsFrame);
 	if (SubFrame > 0 )
 	{
-		SubFrame = FMath::Min(SubFrame, 0.999999940);
+		SubFrame = FMath::Min(SubFrame, (double)FFrameTime::MaxSubframe);
 	}
 
 	return FFrameTime(FrameNumber, SubFrame);
diff --git a/Engine/Source/Runtime/TimeManagement/Public/FrameTime.h b/Engine/Source/Runtime/TimeManagement/Public/FrameTime.h
--- a/Engine/Source/Runtime/TimeManagement/Public/FrameTime.h
+++ b/Engine/Source/Runtime/TimeManagement/Public/FrameTime.h
@@ -12,10 +12,12 @@
  * Only the frame number part of this representation can be negative, sub frames are always a positive value between the frame number and its next logical frame
  */
 USTRUCT(BlueprintType)
-struct FFrameTime
+struct TIMEMANAGEMENT_API FFrameTime
 {
 	GENERATED_BODY()
 
+	static const float MaxSubframe;
+
 	/**
 	 * Default constructor initializing to zero
 	 */
@@ -141,7 +143,7 @@ inline FFrameTime::FFrameTime(FFrameNumber InFrameNumber, float InSubFrame)
 {
 	// Hack to ensure that SubFrames are in a sensible range of precision to work around
 	// problems with FloorToXYZ returning the wrong thing for very small negative numbers
-	SubFrame = FMath::Clamp(SubFrame + 0.5f - 0.5f, 0.f, 0.99999994f);
+	SubFrame = FMath::Clamp(SubFrame + 0.5f - 0.5f, 0.f, MaxSubframe);
 	checkSlow(InSubFrame >= 0.f && InSubFrame < 1.f);
 }
 
@@ -271,5 +273,5 @@ FORCEINLINE_DEBUGGABLE FFrameTime FFrameTime::FromDecimal(double InDecimalFrame)
 
 	// Ensure fractional parts above the highest sub frame float precision do not round to 0.0
 	double Fraction = InDecimalFrame - FMath::FloorToDouble(InDecimalFrame);
-	return FFrameTime(NewFrame, FMath::Clamp(Fraction, 0.0, 0.999999940));
+	return FFrameTime(NewFrame, FMath::Clamp(Fraction, 0.0, (double)MaxSubframe));
 }
\ No newline at end of file
diff --git a/Engine/Source/Runtime/UMG/Private/Animation/MovieScene2DTransformSection.cpp b/Engine/Source/Runtime/UMG/Private/Animation/MovieScene2DTransformSection.cpp
--- a/Engine/Source/Runtime/UMG/Private/Animation/MovieScene2DTransformSection.cpp
+++ b/Engine/Source/Runtime/UMG/Private/Animation/MovieScene2DTransformSection.cpp
@@ -10,7 +10,7 @@
 
 struct F2DTransformSectionEditorData
 {
-	F2DTransformSectionEditorData()
+	F2DTransformSectionEditorData(EMovieScene2DTransformChannel Mask)
 	{
 		FText TranslationGroup = NSLOCTEXT("MovieScene2DTransformSection", "Translation", "Translation");
 		FText RotationGroup = NSLOCTEXT("MovieScene2DTransformSection", "Rotation", "Rotation");
@@ -18,24 +18,31 @@ struct F2DTransformSectionEditorData
 		FText ShearGroup = NSLOCTEXT("MovieScene2DTransformSection", "Shear", "Shear");
 
 		CommonData[0].SetIdentifiers("Translation.X", FCommonChannelData::ChannelX, TranslationGroup);
+		CommonData[0].bEnabled = EnumHasAllFlags(Mask, EMovieScene2DTransformChannel::TranslationX);
 		CommonData[0].SortOrder = 0;
 
 		CommonData[1].SetIdentifiers("Translation.Y", FCommonChannelData::ChannelY, TranslationGroup);
+		CommonData[1].bEnabled = EnumHasAllFlags(Mask, EMovieScene2DTransformChannel::TranslationY);
 		CommonData[1].SortOrder = 1;
 
 		CommonData[2].SetIdentifiers("Angle", NSLOCTEXT("MovieScene2DTransformSection", "AngleText", "Angle"), RotationGroup);
+		CommonData[2].bEnabled = EnumHasAllFlags(Mask, EMovieScene2DTransformChannel::Rotation);
 		CommonData[2].SortOrder = 2;
 
 		CommonData[3].SetIdentifiers("Scale.X", FCommonChannelData::ChannelX, ScaleGroup);
+		CommonData[3].bEnabled = EnumHasAllFlags(Mask, EMovieScene2DTransformChannel::ScaleX);
 		CommonData[3].SortOrder = 3;
 		
 		CommonData[4].SetIdentifiers("Scale.Y", FCommonChannelData::ChannelY, ScaleGroup);
+		CommonData[4].bEnabled = EnumHasAllFlags(Mask, EMovieScene2DTransformChannel::ScaleY);
 		CommonData[4].SortOrder = 4;
 
 		CommonData[5].SetIdentifiers("Shear.X", FCommonChannelData::ChannelX, ShearGroup);
+		CommonData[5].bEnabled = EnumHasAllFlags(Mask, EMovieScene2DTransformChannel::ShearX);
 		CommonData[5].SortOrder = 5;
 		
 		CommonData[6].SetIdentifiers("Shear.Y", FCommonChannelData::ChannelY, ShearGroup);
+		CommonData[6].bEnabled = EnumHasAllFlags(Mask, EMovieScene2DTransformChannel::ShearY);
 		CommonData[6].SortOrder = 6;
 
 		ExternalValues[0].OnGetExternalValue = ExtractTranslationX;
@@ -94,13 +101,28 @@ UMovieScene2DTransformSection::UMovieScene2DTransformSection(const FObjectInitia
 			GetLinkerCustomVersion(FSequencerObjectVersion::GUID) < FSequencerObjectVersion::WhenFinishedDefaultsToProjectDefault ? 
 			EMovieSceneCompletionMode::RestoreState : 
 			EMovieSceneCompletionMode::ProjectDefault);
+
+	ProxyChannels = EMovieScene2DTransformChannel::None;
+	TransformMask = EMovieScene2DTransformChannel::AllTransform;
 	BlendType = EMovieSceneBlendType::Absolute;
 
+	UpdateChannelProxy();
+}
+
+void UMovieScene2DTransformSection::UpdateChannelProxy()
+{
+	if (ProxyChannels == TransformMask.GetChannels())
+	{
+		return;
+	}
+
+	ProxyChannels = TransformMask.GetChannels();
+
 	FMovieSceneChannelData Channels;
 
 #if WITH_EDITOR
 
-	static const F2DTransformSectionEditorData EditorData;
+	F2DTransformSectionEditorData EditorData(TransformMask.GetChannels());
 
 	Channels.Add(Translation[0], EditorData.CommonData[0], EditorData.ExternalValues[0]);
 	Channels.Add(Translation[1], EditorData.CommonData[1], EditorData.ExternalValues[1]);
@@ -124,3 +146,71 @@ UMovieScene2DTransformSection::UMovieScene2DTransformSection(const FObjectInitia
 
 	ChannelProxy = MakeShared<FMovieSceneChannelProxy>(MoveTemp(Channels));
 }
+
+void UMovieScene2DTransformSection::Serialize(FArchive& Ar)
+{
+	Super::Serialize(Ar);
+
+	if (Ar.IsLoading())
+	{
+		UpdateChannelProxy();
+	}
+}
+
+
+FMovieScene2DTransformMask UMovieScene2DTransformSection::GetMask() const
+{
+	return TransformMask;
+}
+
+void UMovieScene2DTransformSection::SetMask(FMovieScene2DTransformMask NewMask)
+{
+	TransformMask = NewMask;
+	UpdateChannelProxy();
+}
+
+FMovieScene2DTransformMask UMovieScene2DTransformSection::GetMaskByName(const FName& InName) const
+{
+	if (InName == TEXT("Translation"))
+	{
+		return EMovieScene2DTransformChannel::Translation;
+	}
+	else if (InName == TEXT("Translation.X"))
+	{
+		return EMovieScene2DTransformChannel::TranslationX;
+	}
+	else if (InName == TEXT("Translation.Y"))
+	{
+		return EMovieScene2DTransformChannel::TranslationY;
+	}
+	else if (InName == TEXT("Angle"))
+	{
+		return EMovieScene2DTransformChannel::Rotation;
+	}
+	else if (InName == TEXT("Scale"))
+	{
+		return EMovieScene2DTransformChannel::Scale;
+	}
+	else if (InName == TEXT("Scale.X"))
+	{
+		return EMovieScene2DTransformChannel::ScaleX;
+	}
+	else if (InName == TEXT("Scale.Y"))
+	{
+		return EMovieScene2DTransformChannel::ScaleY;
+	}
+	else if (InName == TEXT("Shear"))
+	{
+		return EMovieScene2DTransformChannel::Shear;
+	}
+	else if (InName == TEXT("Shear.X"))
+	{
+		return EMovieScene2DTransformChannel::ShearX;
+	}
+	else if (InName == TEXT("Shear.Y"))
+	{
+		return EMovieScene2DTransformChannel::ShearY;
+	}
+
+	return EMovieScene2DTransformChannel::AllTransform;
+}
diff --git a/Engine/Source/Runtime/UMG/Private/Animation/MovieScene2DTransformTemplate.cpp b/Engine/Source/Runtime/UMG/Private/Animation/MovieScene2DTransformTemplate.cpp
--- a/Engine/Source/Runtime/UMG/Private/Animation/MovieScene2DTransformTemplate.cpp
+++ b/Engine/Source/Runtime/UMG/Private/Animation/MovieScene2DTransformTemplate.cpp
@@ -15,17 +15,20 @@ FMovieSceneAnimTypeID GetBlendingDataType<FWidgetTransform>()
 FMovieScene2DTransformSectionTemplate::FMovieScene2DTransformSectionTemplate(const UMovieScene2DTransformSection& Section, const UMovieScenePropertyTrack& Track)
 	: FMovieScenePropertySectionTemplate(Track.GetPropertyName(), Track.GetPropertyPath())
 	, BlendType(Section.GetBlendType().Get())
+	, Mask(Section.GetMask())
 {
-	Translation[0] = Section.Translation[0];
-	Translation[1] = Section.Translation[1];
+	EMovieScene2DTransformChannel MaskChannels = Mask.GetChannels();
 
-	Rotation = Section.Rotation;
+	if (EnumHasAllFlags(MaskChannels, EMovieScene2DTransformChannel::TranslationX))	Translation[0] = Section.Translation[0];
+	if (EnumHasAllFlags(MaskChannels, EMovieScene2DTransformChannel::TranslationY))	Translation[1] = Section.Translation[1];
 
-	Scale[0] = Section.Scale[0];
-	Scale[1] = Section.Scale[1];
+	if (EnumHasAllFlags(MaskChannels, EMovieScene2DTransformChannel::Rotation))		Rotation = Section.Rotation;
 
-	Shear[0] = Section.Shear[0];
-	Shear[1] = Section.Shear[1];
+	if (EnumHasAllFlags(MaskChannels, EMovieScene2DTransformChannel::ScaleX))		Scale[0] = Section.Scale[0];
+	if (EnumHasAllFlags(MaskChannels, EMovieScene2DTransformChannel::ScaleY))		Scale[1] = Section.Scale[1];
+
+	if (EnumHasAllFlags(MaskChannels, EMovieScene2DTransformChannel::ShearX))		Shear[0] = Section.Shear[0];
+	if (EnumHasAllFlags(MaskChannels, EMovieScene2DTransformChannel::ShearY))		Shear[1] = Section.Shear[1];
 }
 
 void FMovieScene2DTransformSectionTemplate::Evaluate(const FMovieSceneEvaluationOperand& Operand, const FMovieSceneContext& Context, const FPersistentEvaluationData& PersistentData, FMovieSceneExecutionTokens& ExecutionTokens) const
@@ -33,26 +36,28 @@ void FMovieScene2DTransformSectionTemplate::Evaluate(const FMovieSceneEvaluation
 	const FFrameTime Time = Context.GetTime();
 	MovieScene::TMultiChannelValue<float, 7> AnimatedData;
 
+	EMovieScene2DTransformChannel ChannelMask = Mask.GetChannels();
+
 	// Only activate channels if the curve has data associated with it
-	auto EvalChannel = [&AnimatedData, Time](uint8 ChanneIndex, const FMovieSceneFloatChannel& Channel)
+	auto EvalChannel = [&AnimatedData, Time, ChannelMask](uint8 ChanneIndex, EMovieScene2DTransformChannel ChannelType, const FMovieSceneFloatChannel& Channel)
 	{
 		float Value = 0.f;
-		if (Channel.Evaluate(Time, Value))
+		if (EnumHasAllFlags(ChannelMask, ChannelType) && Channel.Evaluate(Time, Value))
 		{
 			AnimatedData.Set(ChanneIndex, Value);
 		}
 	};
 
-	EvalChannel(0, Translation[0]);
-	EvalChannel(1, Translation[1]);
+	EvalChannel(0, EMovieScene2DTransformChannel::TranslationX, Translation[0]);
+	EvalChannel(1, EMovieScene2DTransformChannel::TranslationY, Translation[1]);
 
-	EvalChannel(2, Scale[0]);
-	EvalChannel(3, Scale[1]);
+	EvalChannel(2, EMovieScene2DTransformChannel::ScaleX, Scale[0]);
+	EvalChannel(3, EMovieScene2DTransformChannel::ScaleY, Scale[1]);
 
-	EvalChannel(4, Shear[0]);
-	EvalChannel(5, Shear[1]);
+	EvalChannel(4, EMovieScene2DTransformChannel::ShearX, Shear[0]);
+	EvalChannel(5, EMovieScene2DTransformChannel::ShearY, Shear[1]);
 
-	EvalChannel(6, Rotation);
+	EvalChannel(6, EMovieScene2DTransformChannel::Rotation, Rotation);
 
 	if (!AnimatedData.IsEmpty())
 	{
diff --git a/Engine/Source/Runtime/UMG/Private/Animation/MovieScene2DTransformTemplate.h b/Engine/Source/Runtime/UMG/Private/Animation/MovieScene2DTransformTemplate.h
--- a/Engine/Source/Runtime/UMG/Private/Animation/MovieScene2DTransformTemplate.h
+++ b/Engine/Source/Runtime/UMG/Private/Animation/MovieScene2DTransformTemplate.h
@@ -5,6 +5,7 @@
 #include "CoreMinimal.h"
 #include "UObject/ObjectMacros.h"
 #include "Channels/MovieSceneFloatChannel.h"
+#include "Animation/MovieScene2DTransformSection.h"
 #include "Evaluation/MovieSceneEvalTemplate.h"
 #include "Evaluation/MovieScenePropertyTemplate.h"
 #include "Evaluation/Blending/MovieSceneMultiChannelBlending.h"
@@ -49,6 +50,9 @@ private:
 	/** Blending method */
 	UPROPERTY()
 	EMovieSceneBlendType BlendType;
+
+	UPROPERTY()
+	FMovieScene2DTransformMask Mask;
 };
 
 /** Access the unique runtime type identifier for a widget transform. */
diff --git a/Engine/Source/Runtime/UMG/Private/Animation/UMGSequencePlayer.cpp b/Engine/Source/Runtime/UMG/Private/Animation/UMGSequencePlayer.cpp
--- a/Engine/Source/Runtime/UMG/Private/Animation/UMGSequencePlayer.cpp
+++ b/Engine/Source/Runtime/UMG/Private/Animation/UMGSequencePlayer.cpp
@@ -27,7 +27,7 @@ void UUMGSequencePlayer::InitSequencePlayer( UWidgetAnimation& InAnimation, UUse
 
 	// Cache the time range of the sequence to determine when we stop
 	Duration = MovieScene::DiscreteSize(MovieScene->GetPlaybackRange());
-	AnimationResolution = MovieScene->GetFrameResolution();
+	AnimationResolution = MovieScene->GetTickResolution();
 	AbsolutePlaybackStart = MovieScene::DiscreteInclusiveLower(MovieScene->GetPlaybackRange());
 }
 
@@ -122,7 +122,8 @@ void UUMGSequencePlayer::Tick(float DeltaTime)
 
 		if ( bCompleted )
 		{
-			if (RootTemplateInstance.IsValid())
+			// There was a bug where Finish wasn't called, assets were created expecting this behavior. At the time of fixing this bug, the object version was incremented so that newly created assets would finish evaluation.
+			if (RootTemplateInstance.IsValid() && Animation->GetLegacyFinishOnStop())
 			{
 				RootTemplateInstance.Finish(*this);
 			}
diff --git a/Engine/Source/Runtime/UMG/Private/Animation/WidgetAnimation.cpp b/Engine/Source/Runtime/UMG/Private/Animation/WidgetAnimation.cpp
--- a/Engine/Source/Runtime/UMG/Private/Animation/WidgetAnimation.cpp
+++ b/Engine/Source/Runtime/UMG/Private/Animation/WidgetAnimation.cpp
@@ -6,6 +6,7 @@
 #include "Blueprint/UserWidget.h"
 #include "MovieScene.h"
 #include "Components/PanelSlot.h"
+#include "UObject/SequencerObjectVersion.h"
 
 
 #define LOCTEXT_NAMESPACE "UWidgetAnimation"
@@ -19,6 +20,20 @@ UWidgetAnimation::UWidgetAnimation(const FObjectInitializer& ObjectInitializer)
 	, MovieScene(nullptr)
 {
 	bParentContextsAreSignificant = false;
+	bLegacyFinishOnStop = true;
+}
+
+/* UObject interface
+ *****************************************************************************/
+
+void UWidgetAnimation::PostLoad()
+{
+	if (GetLinkerCustomVersion(FSequencerObjectVersion::GUID) < FSequencerObjectVersion::FinishUMGEvaluation)
+	{
+		bLegacyFinishOnStop = false;
+	}
+
+	Super::PostLoad();
 }
 
 
@@ -38,7 +53,7 @@ UWidgetAnimation* UWidgetAnimation::GetNullAnimation()
 		NullAnimation->MovieScene = NewObject<UMovieScene>(NullAnimation, FName("No Animation"));
 		NullAnimation->MovieScene->AddToRoot();
 
-		NullAnimation->MovieScene->SetPlaybackFrameRate(FFrameRate(20, 1));
+		NullAnimation->MovieScene->SetDisplayRate(FFrameRate(20, 1));
 	}
 
 	return NullAnimation;
@@ -49,13 +64,13 @@ UWidgetAnimation* UWidgetAnimation::GetNullAnimation()
 
 float UWidgetAnimation::GetStartTime() const
 {
-	return MovieScene->GetPlaybackRange().GetLowerBoundValue() / MovieScene->GetFrameResolution();
+	return MovieScene->GetPlaybackRange().GetLowerBoundValue() / MovieScene->GetTickResolution();
 }
 
 
 float UWidgetAnimation::GetEndTime() const
 {
-	return MovieScene->GetPlaybackRange().GetUpperBoundValue() / MovieScene->GetFrameResolution();
+	return MovieScene->GetPlaybackRange().GetUpperBoundValue() / MovieScene->GetTickResolution();
 }
 
 
diff --git a/Engine/Source/Runtime/UMG/Public/Animation/MovieScene2DTransformSection.h b/Engine/Source/Runtime/UMG/Public/Animation/MovieScene2DTransformSection.h
--- a/Engine/Source/Runtime/UMG/Public/Animation/MovieScene2DTransformSection.h
+++ b/Engine/Source/Runtime/UMG/Public/Animation/MovieScene2DTransformSection.h
@@ -9,6 +9,86 @@
 #include "MovieScene2DTransformSection.generated.h"
 
 
+enum class EMovieScene2DTransformChannel : uint32
+{
+	None			= 0x000,
+
+	TranslationX 	= 0x001,
+	TranslationY 	= 0x002,
+	Translation 	= TranslationX | TranslationY,
+
+	Rotation 		= 0x004,
+
+	ScaleX 			= 0x008,
+	ScaleY 			= 0x010,
+	Scale 			= ScaleX | ScaleY,
+
+	ShearX 			= 0x020,
+	ShearY 			= 0x040,
+	Shear           = ShearX | ShearY,
+
+	AllTransform	= Translation | Rotation | Scale | Shear
+};
+ENUM_CLASS_FLAGS(EMovieScene2DTransformChannel)
+
+USTRUCT()
+struct FMovieScene2DTransformMask
+{
+	GENERATED_BODY()
+
+	FMovieScene2DTransformMask()
+		: Mask(0)
+	{}
+
+	FMovieScene2DTransformMask(EMovieScene2DTransformChannel Channel)
+		: Mask((__underlying_type(EMovieScene2DTransformChannel))Channel)
+	{}
+
+	EMovieScene2DTransformChannel GetChannels() const
+	{
+		return (EMovieScene2DTransformChannel)Mask;
+	}
+
+	FVector2D GetTranslationFactor() const
+	{
+		EMovieScene2DTransformChannel Channels = GetChannels();
+		return FVector2D(
+			EnumHasAllFlags(Channels, EMovieScene2DTransformChannel::TranslationX) ? 1.f : 0.f,
+			EnumHasAllFlags(Channels, EMovieScene2DTransformChannel::TranslationY) ? 1.f : 0.f);
+	}
+
+	float GetRotationFactor() const
+	{
+		EMovieScene2DTransformChannel Channels = GetChannels();
+		if (EnumHasAllFlags(Channels, EMovieScene2DTransformChannel::Rotation))
+		{
+			return 1.f;
+		}
+		return 0.f;
+	}
+
+	FVector2D GetScaleFactor() const
+	{
+		EMovieScene2DTransformChannel Channels = GetChannels();
+		return FVector2D(
+			EnumHasAllFlags(Channels, EMovieScene2DTransformChannel::ScaleX) ? 1.f : 0.f,
+			EnumHasAllFlags(Channels, EMovieScene2DTransformChannel::ScaleY) ? 1.f : 0.f);
+	}
+
+	FVector2D GetShearFactor() const
+	{
+		EMovieScene2DTransformChannel Channels = GetChannels();
+		return FVector2D(
+			EnumHasAllFlags(Channels, EMovieScene2DTransformChannel::ShearX) ? 1.f : 0.f,
+			EnumHasAllFlags(Channels, EMovieScene2DTransformChannel::ShearY) ? 1.f : 0.f);
+	}
+
+private:
+
+	UPROPERTY()
+	uint32 Mask;
+};
+
 /**
  * A transform section
  */
@@ -19,6 +99,32 @@ class UMovieScene2DTransformSection
 	GENERATED_UCLASS_BODY()
 
 public:
+		
+	/**
+	 * Access the mask that defines which channels this track should animate
+	 */
+	UMG_API FMovieScene2DTransformMask GetMask() const;
+
+	/**
+	 * Set the mask that defines which channels this track should animate
+	 */
+	UMG_API void SetMask(FMovieScene2DTransformMask NewMask);
+
+	/**
+	 * Get the mask by name
+	 */
+	UMG_API FMovieScene2DTransformMask GetMaskByName(const FName& InName) const;
+
+protected:
+
+	virtual void Serialize(FArchive& Ar) override;
+
+	void UpdateChannelProxy();
+
+public:
+
+	UPROPERTY()
+	FMovieScene2DTransformMask TransformMask;
 
 	/** Translation curves*/
 	UPROPERTY()
@@ -35,4 +141,7 @@ public:
 	/** Shear curve */
 	UPROPERTY()
 	FMovieSceneFloatChannel Shear[2];
+
+	/** Unserialized mask that defines the mask of the current channel proxy so we don't needlessly re-create it on post-undo */
+	EMovieScene2DTransformChannel ProxyChannels;
 };
diff --git a/Engine/Source/Runtime/UMG/Public/Animation/WidgetAnimation.h b/Engine/Source/Runtime/UMG/Public/Animation/WidgetAnimation.h
--- a/Engine/Source/Runtime/UMG/Public/Animation/WidgetAnimation.h
+++ b/Engine/Source/Runtime/UMG/Public/Animation/WidgetAnimation.h
@@ -72,6 +72,14 @@ public:
 	/** Get Animation bindings of the animation */
 	const TArray<FWidgetAnimationBinding>& GetBindings() const { return AnimationBindings; }
 
+	/** Whether to finish evaluation on stop */
+	bool GetLegacyFinishOnStop() const { return bLegacyFinishOnStop; }
+
+protected:
+
+	/** Called after this object has been deserialized */
+	virtual void PostLoad() override;
+
 public:
 
 	/** Pointer to the movie scene that controls this animation. */
@@ -81,4 +89,10 @@ public:
 	/**  */
 	UPROPERTY()
 	TArray<FWidgetAnimationBinding> AnimationBindings;
+
+private:
+
+	/** Whether to finish evaluation on stop. This legacy value is to preserve existing asset behavior to NOT finish on stop since content was created with this bug. If this is removed, evaluation should always finish on stop. */
+	UPROPERTY()
+	bool bLegacyFinishOnStop;
 };
diff --git a/Engine/Source/Runtime/UMG/UMG.Build.cs b/Engine/Source/Runtime/UMG/UMG.Build.cs
--- a/Engine/Source/Runtime/UMG/UMG.Build.cs
+++ b/Engine/Source/Runtime/UMG/UMG.Build.cs
@@ -32,7 +32,8 @@ public class UMG : ModuleRules
 				"MovieScene",
                 "MovieSceneTracks",
                 "PropertyPath",
-            }
+				"TimeManagement"
+			}
         );
 
         PrivateIncludePathModuleNames.AddRange(
