Fixes FrameNumber nodes tripping an ensure when using them via Blueprints and fixes FrameNumbers & friends not being properly breakable in BP.

diff --git a/Engine/Source/Runtime/Engine/Classes/Kismet/KismetMathLibrary.h b/Engine/Source/Runtime/Engine/Classes/Kismet/KismetMathLibrary.h
--- a/Engine/Source/Runtime/Engine/Classes/Kismet/KismetMathLibrary.h
+++ b/Engine/Source/Runtime/Engine/Classes/Kismet/KismetMathLibrary.h
@@ -11,6 +11,8 @@
 #include "UObject/Stack.h"
 #include "UObject/ScriptMacros.h"
 #include "Kismet/BlueprintFunctionLibrary.h"
+#include "Misc/QualifiedFrameTime.h"
+
 #include "KismetMathLibrary.generated.h"
 
 // Whether to inline functions at all
@@ -1395,7 +1397,26 @@ class ENGINE_API UKismetMathLibrary : public UBlueprintFunctionLibrary
 	UFUNCTION(BlueprintPure, Category="Math|Timespan")
 	static bool TimespanFromString(FString TimespanString, FTimespan& Result);
 
+	//
+	// Frame Time and Frame Rate Functions
+	//
+
+	/** Creates a FQualifiedFrameTime out of a frame number, frame rate, and optional 0-1 clamped subframe. */
+	UFUNCTION(BlueprintPure, Category = "Time Management", meta = (NativeMakeFunc))
+	static FQualifiedFrameTime MakeQualifiedFrameTime(FFrameNumber Frame, FFrameRate FrameRate, float SubFrame = 0.f);
 
+	/** Breaks a FQualifiedFrameTime into its component parts again. */
+	UFUNCTION(BlueprintPure, Category = "Time Management", meta = (NativeBreakFunc))
+	static void BreakQualifiedFrameTime(const FQualifiedFrameTime& InFrameTime, FFrameNumber& Frame, FFrameRate& FrameRate, float& SubFrame);
+
+	/** Creates a FFrameRate from a Numerator and a Denominator. Enforces that the Denominator is at least one. */
+	UFUNCTION(BlueprintPure, Category = "Time Management", meta = (NativeMakeFunc))
+	static FFrameRate MakeFrameRate(int32 Numerator, int32 Denominator = 1);
+
+	/** Breaks a FFrameRate into a numerator and denominator. */
+	UFUNCTION(BlueprintPure, Category = "Time Management", meta = (NativeBreakFunc))
+	static void BreakFrameRate(const FFrameRate& InFrameRate, int32& Numerator, int32& Denominator);
+	
 	// -- Begin K2 utilities
 
 	/** Converts a byte to a float */
diff --git a/Engine/Source/Runtime/TimeManagement/Public/TimeManagementBlueprintLibrary.h b/Engine/Source/Runtime/TimeManagement/Public/TimeManagementBlueprintLibrary.h
--- a/Engine/Source/Runtime/TimeManagement/Public/TimeManagementBlueprintLibrary.h
+++ b/Engine/Source/Runtime/TimeManagement/Public/TimeManagementBlueprintLibrary.h
@@ -29,10 +29,53 @@ public:
 	UFUNCTION(BlueprintPure, meta = (DisplayName = "Seconds * FrameRate", CompactNodeTitle = "*"), Category = "Utilities|Time Management")
 	static FFrameTime Multiply_SecondsFrameRate(float TimeInSeconds, const FFrameRate& FrameRate);
 
-	/** Converts an Timecode to a string ie: hh:mm:ss */
+	/** Converts an Timecode to a string (hh:mm:ss:ff). If bForceSignDisplay then the number sign will always be prepended instead of just when expressing a negative time. */
 	UFUNCTION(BlueprintPure, meta = (DisplayName = "Timecode to String", BlueprintAutocast), Category = "Utilities|Time Management")
 	static FString Conv_TimecodeToString(const FTimecode& InTimecode, bool bForceSignDisplay = false);
 
+	/** Verifies that this is a valid framerate with a non-zero denominator. */
+	UFUNCTION(BlueprintPure, meta = (DisplayName = "IsValid"), Category = "Utilities|Time Management")
+	static bool IsValid_Framerate(const FFrameRate& InFrameRate);
+
+	/** Checks if this framerate is an even multiple of another framerate, ie: 60 is a multiple of 30, but 59.94 is not. */
+	UFUNCTION(BlueprintPure, meta = (DisplayName = "Is Multiple Of"), Category = "Utilities|Time Management")
+	static bool IsValid_MultipleOf(const FFrameRate& InFrameRate, const FFrameRate& OtherFramerate);
+
+	/** Converts the specified time from one framerate to another framerate. This is useful for converting between tick resolution and display rate. */
+	UFUNCTION(BlueprintPure, meta = (DisplayName = "Transform Frame Time"), Category = "Utilities|Time Management")
+	static FFrameTime TransformTime(const FFrameTime& SourceTime, const FFrameRate& SourceRate, const FFrameRate& DestinationRate);
+
+	/** Snaps the given SourceTime to the nearest frame in the specified Destination Framerate. Useful for determining the nearest frame for another resolution. Returns the frame time in the destination frame rate. */
+	UFUNCTION(BlueprintPure, meta = (DisplayName = "Snap Frame Time"), Category = "Utilities|Time Management")
+	static FFrameTime SnapFrameTimeToRate(const FFrameTime& SourceTime, const FFrameRate& SourceRate, const FFrameRate& SnapToRate);
+
+	/** Addition (FrameNumber A + FrameNumber B) */
+	UFUNCTION(BlueprintPure, meta = (DisplayName = "FrameNumber + FrameNumber", CompactNodeTitle = "+", Keywords = "+ add plus", CommutativeAssociativeBinaryOperator = "true", ScriptMethod, ScriptMethodSelfReturn, ScriptOperator = "+;+="), Category = "Utilities|Time Management")
+	static FFrameNumber Add_FrameNumberFrameNumber(FFrameNumber A, FFrameNumber B);
+
+	/** Subtraction (FrameNumber A - FrameNumber B) */
+	UFUNCTION(BlueprintPure, meta = (DisplayName = "FrameNumber - FrameNumber", CompactNodeTitle = "-", Keywords = "- subtract minus", ScriptMethod, ScriptMethodSelfReturn, ScriptOperator = "-;-="), Category = "Utilities|Time Management")
+	static FFrameNumber Subtract_FrameNumberFrameNumber(FFrameNumber A, FFrameNumber B);
+
+	/** Addition (FrameNumber A + int B) */
+	UFUNCTION(BlueprintPure, meta = (DisplayName = "FrameNumber + Int", CompactNodeTitle = "+", Keywords = "+ add plus", ScriptMethod, ScriptMethodSelfReturn, ScriptOperator = "+;+="), Category = "Utilities|Time Management")
+	static FFrameNumber Add_FrameNumberInteger(FFrameNumber A, int32 B);
+
+	/** Subtraction (FrameNumber A - int B) */
+	UFUNCTION(BlueprintPure, meta = (DisplayName = "FrameNumber - Int", CompactNodeTitle = "-", Keywords = "- subtract minus", ScriptMethod, ScriptMethodSelfReturn, ScriptOperator = "-;-="), Category = "Utilities|Time Management")
+	static FFrameNumber Subtract_FrameNumberInteger(FFrameNumber A, int32 B);
+
+	/** Multiply (FrameNumber A * B) */
+	UFUNCTION(BlueprintPure, meta = (DisplayName = "FrameNumber * Int", CompactNodeTitle = "*", Keywords = "* multiply", ScriptMethod, ScriptMethodSelfReturn, ScriptOperator = "*;*="), Category = "Utilities|Time Management")
+	static FFrameNumber Multiply_FrameNumberInteger(FFrameNumber A, int32 B);
+
+	/** Divide (FrameNumber A / B) */
+	UFUNCTION(BlueprintPure, meta = (DisplayName = "FrameNumber - FrameNumber", CompactNodeTitle = "-", Keywords = "/ divide", ScriptMethod, ScriptMethodSelfReturn, ScriptOperator = "/;/="), Category = "Utilities|Time Management")
+	static FFrameNumber Divide_FrameNumberInteger(FFrameNumber A, int32 B);
+	
+	/** Converts a FrameNumber to an int32 for use in functions that take int32 frame counts for convenience. */
+	UFUNCTION(BlueprintPure, meta = (DisplayName = "FrameNumber to Integer", ScriptName="FrameNumberToInteger", BlueprintAutocast), Category = "Utilities|Time Management")
+	static int32 Conv_FrameNumberToInteger(const FFrameNumber& InFrameNumber);
 public:
 	/**
 	 * Get the Timecode from the TimeManagement's TimecodeProvider.
