Changed action checked for frame rate radio button to check against value attribute instead of current frame rate because value attribute was the one being set when a frame rate was selected.

diff --git a/Engine/Source/Runtime/TimeManagement/Private/Widgets/SFrameRatePicker.cpp b/Engine/Source/Runtime/TimeManagement/Private/Widgets/SFrameRatePicker.cpp
--- a/Engine/Source/Runtime/TimeManagement/Private/Widgets/SFrameRatePicker.cpp
+++ b/Engine/Source/Runtime/TimeManagement/Private/Widgets/SFrameRatePicker.cpp
@@ -100,7 +100,7 @@ TSharedRef<SWidget> SFrameRatePicker::BuildMenu()
 			FUIAction MenuAction(
 				FExecuteAction::CreateSP(this, &SFrameRatePicker::SetValue, ThisFrameRate),
 				FCanExecuteAction(),
-				FIsActionChecked::CreateLambda([=]{ return CurrentFrameRate.IsSet() && CurrentFrameRate.GetValue() == ThisFrameRate; })
+				FIsActionChecked::CreateLambda([=]{ return ValueAttribute.IsSet() && ValueAttribute.Get() == ThisFrameRate; })
 				);
 			MenuBuilder.AddMenuEntry(Preset.DisplayName, Preset.Description, FSlateIcon(), MenuAction, NAME_None, EUserInterfaceActionType::RadioButton);
 		}
@@ -133,7 +133,7 @@ TSharedRef<SWidget> SFrameRatePicker::BuildMenu()
 				FUIAction MenuAction(
 					FExecuteAction::CreateSP(this, &SFrameRatePicker::SetValue, ThisFrameRate),
 					FCanExecuteAction(),
-					FIsActionChecked::CreateLambda([=]{ return CurrentFrameRate.IsSet() && CurrentFrameRate.GetValue() == ThisFrameRate; })
+					FIsActionChecked::CreateLambda([=]{ return ValueAttribute.IsSet() && ValueAttribute.Get() == ThisFrameRate; })
 					);
 				MenuBuilder.AddMenuEntry(Preset.DisplayName, Preset.Description, FSlateIcon(), MenuAction, NAME_None, EUserInterfaceActionType::RadioButton);
 			}
@@ -182,7 +182,7 @@ void SFrameRatePicker::PopulateNotRecommendedMenu(FMenuBuilder& MenuBuilder)
 		FUIAction MenuAction(
 			FExecuteAction::CreateSP(this, &SFrameRatePicker::SetValue, ThisFrameRate),
 			FCanExecuteAction(),
-			FIsActionChecked::CreateLambda([=]{ return CurrentFrameRate.IsSet() && CurrentFrameRate.GetValue() == ThisFrameRate; })
+			FIsActionChecked::CreateLambda([=]{ return ValueAttribute.IsSet() && ValueAttribute.Get() == ThisFrameRate; })
 			);
 		MenuBuilder.AddMenuEntry(Preset.DisplayName, Preset.Description, FSlateIcon(), MenuAction, NAME_None, EUserInterfaceActionType::RadioButton);
 
