Fix incorrect CompactNodeTitle on FFrameNumber Divide_FrameNumberInteger

#rnx

diff --git a/Engine/Source/Runtime/TimeManagement/Public/TimeManagementBlueprintLibrary.h b/Engine/Source/Runtime/TimeManagement/Public/TimeManagementBlueprintLibrary.h
--- a/Engine/Source/Runtime/TimeManagement/Public/TimeManagementBlueprintLibrary.h
+++ b/Engine/Source/Runtime/TimeManagement/Public/TimeManagementBlueprintLibrary.h
@@ -71,7 +71,7 @@ public:
 	static FFrameNumber Multiply_FrameNumberInteger(FFrameNumber A, int32 B);
 
 	/** Divide (FrameNumber A / B) */
-	UFUNCTION(BlueprintPure, meta = (DisplayName = "FrameNumber - FrameNumber", CompactNodeTitle = "-", Keywords = "/ divide", ScriptMethod, ScriptMethodSelfReturn, ScriptOperator = "/;/="), Category = "Utilities|Time Management")
+	UFUNCTION(BlueprintPure, meta = (DisplayName = "FrameNumber / FrameNumber", CompactNodeTitle = "/", Keywords = "/ divide", ScriptMethod, ScriptMethodSelfReturn, ScriptOperator = "/;/="), Category = "Utilities|Time Management")
 	static FFrameNumber Divide_FrameNumberInteger(FFrameNumber A, int32 B);
 	
 	/** Converts a FrameNumber to an int32 for use in functions that take int32 frame counts for convenience. */
