Fixed Unreadable Text in FPS Advanced Options Dropdown

diff --git a/Engine/Source/Runtime/TimeManagement/Private/Widgets/SFrameRatePicker.cpp b/Engine/Source/Runtime/TimeManagement/Private/Widgets/SFrameRatePicker.cpp
--- a/Engine/Source/Runtime/TimeManagement/Private/Widgets/SFrameRatePicker.cpp
+++ b/Engine/Source/Runtime/TimeManagement/Private/Widgets/SFrameRatePicker.cpp
@@ -37,7 +37,6 @@ void SFrameRatePicker::Construct(const FArguments& InArgs)
 		.ButtonStyle(InArgs._ButtonStyle)
 		.ForegroundColor(InArgs._ForegroundColor)
 		.ContentPadding(InArgs._ContentPadding)
-		.VAlign(VAlign_Fill)
 		.OnGetMenuContent(this, &SFrameRatePicker::BuildMenu)
 		.ButtonContent()
 		[
diff --git a/Engine/Source/Runtime/TimeManagement/Public/Widgets/SFrameRatePicker.h b/Engine/Source/Runtime/TimeManagement/Public/Widgets/SFrameRatePicker.h
--- a/Engine/Source/Runtime/TimeManagement/Public/Widgets/SFrameRatePicker.h
+++ b/Engine/Source/Runtime/TimeManagement/Public/Widgets/SFrameRatePicker.h
@@ -24,7 +24,7 @@ public:
 	SLATE_BEGIN_ARGS(SFrameRatePicker)
 		: _ComboButtonStyle(&FCoreStyle::Get().GetWidgetStyle< FComboButtonStyle >("ComboButton"))
 		, _ButtonStyle(nullptr)
-		, _ForegroundColor(FCoreStyle::Get().GetSlateColor("InvertedForeground"))
+		, _ForegroundColor(FSlateColor::UseForeground())
 		, _ContentPadding(FMargin(2.f, 0.f))
 		, _HasMultipleValues(false)
 		, _PresetValues()
