TimeManagement: add FCommonFrameRates::Find() overload to allow frame rate lookup by FPS

This new Find() overload allows finding a common frame rate given its value as a decimal number of
frames per second.

diff --git a/Engine/Source/Runtime/TimeManagement/Private/CommonFrameRates.cpp b/Engine/Source/Runtime/TimeManagement/Private/CommonFrameRates.cpp
--- a/Engine/Source/Runtime/TimeManagement/Private/CommonFrameRates.cpp
+++ b/Engine/Source/Runtime/TimeManagement/Private/CommonFrameRates.cpp
@@ -32,4 +32,13 @@ const FCommonFrameRateInfo* FCommonFrameRates::Find(FFrameRate InFrameRate)
 	return Algo::FindBy(AllFrameRates, InFrameRate, &FCommonFrameRateInfo::FrameRate);
 }
 
+const FCommonFrameRateInfo* FCommonFrameRates::Find(const double InFrameRateAsDecimal, const double Tolerance)
+{
+	return Algo::FindByPredicate(AllFrameRates,
+		[InFrameRateAsDecimal, Tolerance](const FCommonFrameRateInfo& CommonFrameRateInfo) -> bool
+		{
+			return FMath::IsNearlyEqual(CommonFrameRateInfo.FrameRate.AsDecimal(), InFrameRateAsDecimal, Tolerance);
+		});
+}
+
 #undef LOCTEXT_NAMESPACE
\ No newline at end of file
diff --git a/Engine/Source/Runtime/TimeManagement/Public/CommonFrameRates.h b/Engine/Source/Runtime/TimeManagement/Public/CommonFrameRates.h
--- a/Engine/Source/Runtime/TimeManagement/Public/CommonFrameRates.h
+++ b/Engine/Source/Runtime/TimeManagement/Public/CommonFrameRates.h
@@ -61,6 +61,14 @@ struct TIMEMANAGEMENT_API FCommonFrameRates
 
 	static const FCommonFrameRateInfo* Find(FFrameRate InFrameRate);
 
+	/** Find a common frame rate that matches the given frame rate as a decimal number of frames per second.
+	 *
+	 *  @param InFrameRateAsDecimal: Frame rate (in frames per second) to search for.
+	 *  @param Tolerance: Numerical tolerance to use when searching for a frame rate match.
+	 *  @return: a pointer to the matching common frame rate if a match was found, or nullptr otherwise.
+	 */
+	static const FCommonFrameRateInfo* Find(const double InFrameRateAsDecimal, const double Tolerance = DOUBLE_KINDA_SMALL_NUMBER);
+
 private:
 	static const FCommonFrameRateInfo AllFrameRates[(int32)ECommonFrameRate::Private_Num];
 };
\ No newline at end of file
