Add API allowing for FFramePicker widget creation outside of TimeManagementModule

diff --git a/Engine/Source/Runtime/TimeManagement/Private/TimeManagementModule.cpp b/Engine/Source/Runtime/TimeManagement/Private/TimeManagementModule.cpp
--- a/Engine/Source/Runtime/TimeManagement/Private/TimeManagementModule.cpp
+++ b/Engine/Source/Runtime/TimeManagement/Private/TimeManagementModule.cpp
@@ -10,6 +10,11 @@ public:
 	virtual FTimedDataInputCollection& GetTimedDataInputCollection() { return Collection; }
 	virtual TArrayView<const struct FCommonFrameRateInfo> GetAllCommonFrameRates() { return FCommonFrameRates::GetAll(); }
 
+	virtual TSharedRef<SFrameRatePicker> CreateFrameRatePicker(SFrameRatePicker::FArguments Arguments) override
+	{
+		return SArgumentNew(Arguments, SFrameRatePicker);
+	}
+	
 private:
 	FTimedDataInputCollection Collection;
 };
diff --git a/Engine/Source/Runtime/TimeManagement/Public/ITimeManagementModule.h b/Engine/Source/Runtime/TimeManagement/Public/ITimeManagementModule.h
--- a/Engine/Source/Runtime/TimeManagement/Public/ITimeManagementModule.h
+++ b/Engine/Source/Runtime/TimeManagement/Public/ITimeManagementModule.h
@@ -5,7 +5,7 @@
 #include "CoreMinimal.h"
 #include "Modules/ModuleInterface.h"
 #include "Modules/ModuleManager.h"
-
+#include "Widgets/SFrameRatePicker.h"
 
 #define TIMEMANAGEMENT_MODULE_NAME TEXT("TimeManagement")
 
@@ -38,4 +38,7 @@ public:
 
 	/** Returns all stored common frame rates */
 	virtual TArrayView<const struct FCommonFrameRateInfo> GetAllCommonFrameRates() = 0;
+
+	/** Returns a widget allowing for the user to pick a specific frame rate */
+	virtual TSharedRef<SFrameRatePicker> CreateFrameRatePicker(SFrameRatePicker::FArguments Arguments) = 0;
 };
