TimeManagement: enable customization of the subframe indicator when stringifying time values

By default, FFrameNumberInterface uses an asterisk suffix to indicate that a stringified time
value represents a subframe time between whole number frames. This change allows
customizing that suffix for clients that can provide more useful information about the subframe
value. This will be used by skeletal animation tracks that have subframe timecode metadata.

diff --git a/Engine/Source/Runtime/TimeManagement/Public/FrameNumberNumericInterface.h b/Engine/Source/Runtime/TimeManagement/Public/FrameNumberNumericInterface.h
--- a/Engine/Source/Runtime/TimeManagement/Public/FrameNumberNumericInterface.h
+++ b/Engine/Source/Runtime/TimeManagement/Public/FrameNumberNumericInterface.h
@@ -34,6 +34,7 @@ struct FFrameNumberInterface : public INumericTypeInterface<double>
 		, TickResolutionAttr(InTickResolutionAttr)
 		, DisplayRateAttr(InDisplayRateAttr)
 		, ZeroPadFramesAttr(InOnGetZeroPadFrameNumber)
+		, SubframeIndicator(TEXT("*"))
 	{
 		check(InDisplayFormatAttr.IsSet());
 		check(InTickResolutionAttr.IsSet());
@@ -89,9 +90,9 @@ struct FFrameNumberInterface : public INumericTypeInterface<double>
 		{
 			// Convert from sequence resolution into display rate frames.
 			FFrameTime DisplayTime = FFrameRate::TransformTime(FFrameTime::FromDecimal(Value), SourceFrameRate, DestinationFrameRate);
-			FString SubframeIndicator = FMath::IsNearlyZero(DisplayTime.GetSubFrame()) ? TEXT("") : TEXT("*");
+			FString SubframeSuffix = FMath::IsNearlyZero(DisplayTime.GetSubFrame()) ? TEXT("") : SubframeIndicator;
 
-			return FString::Printf(TEXT("%0*d%s"), ZeroPadFramesAttr.Get(), DisplayTime.GetFrame().Value, *SubframeIndicator);
+			return FString::Printf(TEXT("%0*d%s"), ZeroPadFramesAttr.Get(), DisplayTime.GetFrame().Value, *SubframeSuffix);
 		}
 		case EFrameNumberDisplayFormats::Seconds:
 		{
@@ -108,12 +109,12 @@ struct FFrameNumberInterface : public INumericTypeInterface<double>
 		{
 			FFrameTime InternalFrameNumber = FFrameTime::FromDecimal(Value);
 			FFrameTime DisplayTime = FFrameRate::TransformTime(InternalFrameNumber, SourceFrameRate, DestinationFrameRate);
-			FString SubframeIndicator = FMath::IsNearlyZero(DisplayTime.GetSubFrame()) ? TEXT("") : TEXT("*");
+			FString SubframeSuffix = FMath::IsNearlyZero(DisplayTime.GetSubFrame()) ? TEXT("") : SubframeIndicator;
 
 			bool bIsDropTimecode = Format == EFrameNumberDisplayFormats::DropFrameTimecode;
 
 			FTimecode AsNonDropTimecode = FTimecode::FromFrameNumber(DisplayTime.FloorToFrame(), DestinationFrameRate, bIsDropTimecode);
-			return FString::Printf(TEXT("[%s%s]"), *AsNonDropTimecode.ToString(false), *SubframeIndicator);
+			return FString::Printf(TEXT("[%s%s]"), *AsNonDropTimecode.ToString(false), *SubframeSuffix);
 		}
 		default:
 			return FString(TEXT("Unsupported Format"));
@@ -210,9 +211,23 @@ struct FFrameNumberInterface : public INumericTypeInterface<double>
 		return TOptional<double>();
 	}
 
+	/** Gets the string suffix used when a stringified value represents a time between whole frames. */
+	FString GetSubframeIndicator() const
+	{
+		return SubframeIndicator;
+	}
+
+	/** Sets the string suffix to use when a stringified value represents a time between whole frames. */
+	void SetSubframeIndicator(const FString& InSubframeIndicator)
+	{
+		SubframeIndicator = InSubframeIndicator;
+	}
+
 private:
 	TAttribute<EFrameNumberDisplayFormats> DisplayFormatAttr;
 	TAttribute<FFrameRate> TickResolutionAttr;
 	TAttribute<FFrameRate> DisplayRateAttr;
 	TAttribute<uint8> ZeroPadFramesAttr;
-};
\ No newline at end of file
+
+	FString SubframeIndicator;
+};
