Incorporate truncated value of SubFrame calculation when converting double seconds value to FrameTime/Number

diff --git a/Engine/Source/Runtime/Core/Public/Misc/FrameRate.h b/Engine/Source/Runtime/Core/Public/Misc/FrameRate.h
--- a/Engine/Source/Runtime/Core/Public/Misc/FrameRate.h
+++ b/Engine/Source/Runtime/Core/Public/Misc/FrameRate.h
@@ -169,15 +169,16 @@ inline double FFrameRate::AsSeconds(FFrameTime FrameTime) const
 	return (double(IntegerPart) + FloatPart) / Numerator;
 }
 
-
-
 inline FFrameTime FFrameRate::AsFrameTime(double TimeInSeconds) const
 {
 	// @todo: sequencer-timecode: proper large number integer multiplication/division before coercion to float ?
-	const double       TimeAsFrame = (TimeInSeconds * Numerator) / Denominator;
-	const FFrameNumber FrameNumber = static_cast<int32>(FMath::Clamp(FMath::FloorToDouble(TimeAsFrame), (double)TNumericLimits<int32>::Min(), (double)TNumericLimits<int32>::Max()));
+	const double TimeAsFrame = (TimeInSeconds * Numerator) / Denominator;
+	FFrameNumber FrameNumber = static_cast<int32>(FMath::Clamp(FMath::FloorToDouble(TimeAsFrame), static_cast<double>(TNumericLimits<int32>::Min()), static_cast<double>(TNumericLimits<int32>::Max())));
 
 	float SubFrame = static_cast<float>(TimeAsFrame - FMath::FloorToDouble(TimeAsFrame));
+	const int32 TruncatedSubFrame = FMath::TruncToInt(SubFrame);
+	SubFrame -= static_cast<float>(TruncatedSubFrame);
+	FrameNumber.Value += TruncatedSubFrame;
 	if (SubFrame > 0.f)
 	{
 		SubFrame = FMath::Min(SubFrame, FFrameTime::MaxSubframe);
@@ -189,8 +190,8 @@ inline FFrameTime FFrameRate::AsFrameTime(double TimeInSeconds) const
 inline FFrameNumber FFrameRate::AsFrameNumber(double TimeInSeconds) const
 {
 	// @todo: sequencer-timecode: proper large number integer multiplication/division before coercion to float ?
-	const double       TimeAsFrame = (double(TimeInSeconds) * Numerator) / Denominator;
-	return static_cast<int32>(FMath::Clamp(FMath::FloorToDouble(TimeAsFrame), (double)TNumericLimits<int32>::Min(), (double)TNumericLimits<int32>::Max()));
+	const double TimeAsFrame = (static_cast<double>(TimeInSeconds) * Numerator) / Denominator;
+	return static_cast<int32>(FMath::Clamp(FMath::FloorToDouble(TimeAsFrame + static_cast<double>(FMath::TruncToInt(static_cast<float>(TimeAsFrame - FMath::FloorToDouble(TimeAsFrame))))), static_cast<double>(TNumericLimits<int32>::Min()), static_cast<double>(TNumericLimits<int32>::Max())));
 }
 
 inline bool operator==(const FFrameRate& A, const FFrameRate& B)
diff --git a/Engine/Source/Runtime/TimeManagement/Private/TimeManagementTests.cpp b/Engine/Source/Runtime/TimeManagement/Private/TimeManagementTests.cpp
--- a/Engine/Source/Runtime/TimeManagement/Private/TimeManagementTests.cpp
+++ b/Engine/Source/Runtime/TimeManagement/Private/TimeManagementTests.cpp
@@ -559,4 +559,70 @@ bool FFrameRateMultiplesTest::RunTest(const FString& Parameters)
 	return true;
 }
 
+IMPLEMENT_SIMPLE_AUTOMATION_TEST(FFrameTimeToSecondsConversionTest, "System.Core.Time.FrameConversion", EAutomationTestFlags::ApplicationContextMask | EAutomationTestFlags::EngineFilter | EAutomationTestFlags::HighPriority)
+bool FFrameTimeToSecondsConversionTest::RunTest(const FString& Parameters)
+{
+	FFrameRate TestRates[] = {
+		FCommonFrameRates::FPS_12(),
+		FCommonFrameRates::FPS_15(),
+		FCommonFrameRates::FPS_24(),
+		FCommonFrameRates::FPS_25(),
+		FCommonFrameRates::FPS_30(),
+		FCommonFrameRates::FPS_48(),
+		FCommonFrameRates::FPS_50(),
+		FCommonFrameRates::FPS_60(),
+		FCommonFrameRates::FPS_100(),
+		FCommonFrameRates::FPS_120(),
+		FCommonFrameRates::FPS_240(),
+		FCommonFrameRates::NTSC_24(),
+		FCommonFrameRates::NTSC_30(),
+		FCommonFrameRates::NTSC_60(),
+		FFrameRate(24000,1)
+	};
+	
+	const int32 NumberOfSecondsToTest = 60;
+	for (const FFrameRate& FrameRate : TestRates)
+	{
+		const int32 NumberOfFrames = FrameRate.AsFrameNumber(static_cast<double>(NumberOfSecondsToTest)).Value;
+		const double Interval = FrameRate.AsInterval();
+		
+		for (int32 FrameIndex = 0; FrameIndex < NumberOfFrames; ++FrameIndex)
+		{
+			// Get seconds representation according to frame-rate
+			const double FrameAsSeconds = FrameRate.AsSeconds(FrameIndex);
+			const double IntervalSeconds = static_cast<double>(FrameIndex) * Interval;
+
+			// Verify seconds value against interval * FrameIndex
+			ensureAlwaysMsgf(FMath::IsNearlyEqual(IntervalSeconds, FrameAsSeconds),
+						TEXT("Did not expect value %d ::AsSeconds for input %i (%d)."),
+						FrameAsSeconds, FrameIndex, IntervalSeconds
+					);
+			
+			// Convert seconds values to FrameTime/FrameNumber values
+			const FFrameTime FrameTime = FrameRate.AsFrameTime(FrameAsSeconds);
+			const FFrameNumber FrameNumber = FrameRate.AsFrameNumber(FrameAsSeconds);
+
+			// Verify frame numbers match original FrameIndex value
+			ensureAlwaysMsgf(FrameTime.GetFrame() == FrameIndex,
+						TEXT("Did not expect frame number %i ::AsFrameTime for input %i (%d)."),
+						FrameNumber.Value, FrameIndex, FrameAsSeconds
+					);
+			
+			ensureAlwaysMsgf(FrameNumber.Value == FrameIndex,
+						TEXT("Did not expect frame number %i ::AsFrameNumber for input %i (%d)."),
+						FrameNumber.Value, FrameIndex, FrameAsSeconds
+					);
+			
+			// Convert FrameTime back to seconds and verify against value used to generate it
+			const double ConvertedSeconds = FrameRate.AsSeconds(FrameTime);
+			ensureAlwaysMsgf(FMath::IsNearlyEqual(ConvertedSeconds, FrameAsSeconds),
+				TEXT("Did not expect value %d ::AsSeconds for input %i (%d) - %s."),
+				ConvertedSeconds, FrameIndex, FrameAsSeconds, *LexToString(FrameTime)
+			);
+		}
+	}
+
+	return true;
+}
+
 #endif // WITH_DEV_AUTOMATION_TESTS
\ No newline at end of file
