CustomTimeStep: Added more options to GenlockedCustomTimestep.

* Option to disabling blocking. This is useful when nDisplay is already controlling frame rate. If frame rate is not enforced in some other way, game time will seem to pass at the wrong rate.
* Force single frame delta time. Normally this custom timestemp quantizes the delta time to a whole number that keeps up with real time. With this option, delta time will always be equal to one frame time, regardless to hom much real time has passed.

Defaults leave old behavior unchanged.

diff --git a/Engine/Source/Runtime/TimeManagement/Private/GenlockedFixedRateCustomTimeStep.cpp b/Engine/Source/Runtime/TimeManagement/Private/GenlockedFixedRateCustomTimeStep.cpp
--- a/Engine/Source/Runtime/TimeManagement/Private/GenlockedFixedRateCustomTimeStep.cpp
+++ b/Engine/Source/Runtime/TimeManagement/Private/GenlockedFixedRateCustomTimeStep.cpp
@@ -6,6 +6,8 @@
 UGenlockedFixedRateCustomTimeStep::UGenlockedFixedRateCustomTimeStep(const FObjectInitializer& ObjectInitializer)
 	: Super(ObjectInitializer)
 	, FrameRate(24,1)
+	, bShouldBlock(true)
+	, bForceSingleFrameDeltaTime(false)
 	, LastSyncCountDelta(0)
 	, QuantizedCurrentTime(0.0)
 {
@@ -68,30 +70,33 @@ bool UGenlockedFixedRateCustomTimeStep::WaitForSync()
 
 	LastIdleTime = FramePeriod - DeltaRealTime;
 
-	// Sleep during the idle time
-	if (LastIdleTime > 0.f)
+	if (bShouldBlock)
 	{
-		// Normal sleep for the bulk of the idle time.
-		if (LastIdleTime > 5.f / 1000.f)
+		// Sleep during the idle time
+		if (LastIdleTime > 0.f)
 		{
-			FPlatformProcess::SleepNoStats(LastIdleTime - 0.002f);
-		}
+			// Normal sleep for the bulk of the idle time.
+			if (LastIdleTime > 5.f / 1000.f)
+			{
+				FPlatformProcess::SleepNoStats(LastIdleTime - 0.002f);
+			}
 
-		// Give up timeslice for small remainder of wait time.
+			// Give up timeslice for small remainder of wait time.
 
-		const double WaitEndTime = FApp::GetLastTime() + FramePeriod;
+			const double WaitEndTime = FApp::GetLastTime() + FramePeriod;
 
-		while (FPlatformTime::Seconds() < WaitEndTime)
-		{
-			FPlatformProcess::SleepNoStats(0.f);
-		}
+			while (FPlatformTime::Seconds() < WaitEndTime)
+			{
+				FPlatformProcess::SleepNoStats(0.f);
+			}
 
-		// Current platform time should now be right after the desired WaitEndTime, with an overshoot
-		CurrentPlatformTime = FPlatformTime::Seconds();
-		FApp::SetIdleTimeOvershoot(CurrentPlatformTime - WaitEndTime);
+			// Current platform time should now be right after the desired WaitEndTime, with an overshoot
+			CurrentPlatformTime = FPlatformTime::Seconds();
+			FApp::SetIdleTimeOvershoot(CurrentPlatformTime - WaitEndTime);
 
-		// Update DeltaRealTime now that we've slept enough
-		DeltaRealTime = CurrentPlatformTime - FApp::GetLastTime();
+			// Update DeltaRealTime now that we've slept enough
+			DeltaRealTime = CurrentPlatformTime - FApp::GetLastTime();
+		}
 	}
 
 	// This +1e-4 avoids a case of LastSyncCountData incorrectly ending up as 0.
@@ -99,7 +104,20 @@ bool UGenlockedFixedRateCustomTimeStep::WaitForSync()
 
 	// Quantize elapsed frames, capped to the maximum that the integer type can hold.
 	LastSyncCountDelta = uint32(FMath::Min(FMath::Floor(DeltaRealTime / FramePeriod), double(MAX_uint32)));
-	ensure(LastSyncCountDelta > 0);
+
+	if (bShouldBlock)
+	{
+		ensure(LastSyncCountDelta > 0);
+	}
+	else if (LastSyncCountDelta < 1)
+	{
+		LastSyncCountDelta = 1;
+	}
+
+	if (bForceSingleFrameDeltaTime)
+	{
+		LastSyncCountDelta = 1;
+	}
 
 	// Save quantized current time for use outside this function.
 	QuantizedCurrentTime = FApp::GetLastTime() + LastSyncCountDelta * FramePeriod;
diff --git a/Engine/Source/Runtime/TimeManagement/Public/GenlockedFixedRateCustomTimeStep.h b/Engine/Source/Runtime/TimeManagement/Public/GenlockedFixedRateCustomTimeStep.h
--- a/Engine/Source/Runtime/TimeManagement/Public/GenlockedFixedRateCustomTimeStep.h
+++ b/Engine/Source/Runtime/TimeManagement/Public/GenlockedFixedRateCustomTimeStep.h
@@ -53,6 +53,14 @@ public:
 	UPROPERTY(EditAnywhere, Category = "Timing")
 	FFrameRate FrameRate;
 
+	/** Indicates that this custom time step should block to enforce the specified frame rate. Set to false if this is enforced elsewhere. */
+	UPROPERTY(EditAnywhere, Category = "Timing")
+	bool bShouldBlock;
+
+	/** When true, delta time will always be 1/FrameRate, regardless of how much real time has elapsed */
+	UPROPERTY(EditAnywhere, Category = "Timing")
+	bool bForceSingleFrameDeltaTime;
+
 private:
 	uint32 LastSyncCountDelta;
 	double QuantizedCurrentTime;
