Time Management: Set seconds minimum fractional digits to the zero pad frames value.

This fixes flickering when displaying time as seconds and values go from 0.20 to 0.2423 and 0.3, etc. Now they will have the same padding, ie. 0.2000, 0.2423, 0.3000

diff --git a/Engine/Source/Runtime/TimeManagement/Public/FrameNumberNumericInterface.h b/Engine/Source/Runtime/TimeManagement/Public/FrameNumberNumericInterface.h
--- a/Engine/Source/Runtime/TimeManagement/Public/FrameNumberNumericInterface.h
+++ b/Engine/Source/Runtime/TimeManagement/Public/FrameNumberNumericInterface.h
@@ -100,7 +100,7 @@ struct FFrameNumberInterface : public INumericTypeInterface<double>
 
 			static const FNumberFormattingOptions NumberFormattingOptions = FNumberFormattingOptions()
 				.SetUseGrouping(false)
-				.SetMinimumFractionalDigits(2)
+				.SetMinimumFractionalDigits(ZeroPadFramesAttr.Get())
 				.SetMaximumFractionalDigits(ZeroPadFramesAttr.Get());
 			return FastDecimalFormat::NumberToString(TimeInSeconds, ExpressionParser::GetLocalizedNumberFormattingRules(), NumberFormattingOptions);
 		}
