UTimeManagementBlueprintLibrary - deprecated Conv_FrameRateToSeconds, replaced with Conv_FrameRateToInterval, the former was incorrectly implemented

diff --git a/Engine/Source/Runtime/TimeManagement/Private/TimeManagementBlueprintLibrary.cpp b/Engine/Source/Runtime/TimeManagement/Private/TimeManagementBlueprintLibrary.cpp
--- a/Engine/Source/Runtime/TimeManagement/Private/TimeManagementBlueprintLibrary.cpp
+++ b/Engine/Source/Runtime/TimeManagement/Private/TimeManagementBlueprintLibrary.cpp
@@ -8,10 +8,17 @@
 
 float UTimeManagementBlueprintLibrary::Conv_FrameRateToSeconds(const FFrameRate& InFrameRate)
 {
+	// @note: this has been deprecated as it returns FPS, not "SPF" as the description indicates it should - see Conv_FrameRateToInterval below
 	// Accept the loss of precision from conversion when in use with Blueprints.
 	return (float)InFrameRate.AsDecimal();
 }
 
+float UTimeManagementBlueprintLibrary::Conv_FrameRateToInterval(const FFrameRate InFrameRate)
+{
+	// Accept the loss of precision from conversion when in use with Blueprints.
+	return (float)InFrameRate.AsInterval();
+}
+
 float UTimeManagementBlueprintLibrary::Conv_QualifiedFrameTimeToSeconds(const FQualifiedFrameTime& InFrameTime)
 {
 	// Accept the loss of precision from conversion when in use with Blueprints.
diff --git a/Engine/Source/Runtime/TimeManagement/Public/TimeManagementBlueprintLibrary.h b/Engine/Source/Runtime/TimeManagement/Public/TimeManagementBlueprintLibrary.h
--- a/Engine/Source/Runtime/TimeManagement/Public/TimeManagementBlueprintLibrary.h
+++ b/Engine/Source/Runtime/TimeManagement/Public/TimeManagementBlueprintLibrary.h
@@ -30,9 +30,14 @@ class TIMEMANAGEMENT_API UTimeManagementBlueprintLibrary : public UBlueprintFunc
 
 public:
 	/** Converts an FrameRate to a float ie: 1/30 returns 0.0333333 */
-	UFUNCTION(BlueprintPure, meta = (DisplayName = "FrameRate To Seconds", BlueprintAutocast), Category = "Utilities|Time Management")
+	UE_DEPRECATED(5.3, "Conv_FrameRateToSeconds has been deprecated, use Conv_FrameRateToInterval instead")
+	UFUNCTION(BlueprintPure, meta = (DisplayName = "FrameRate To Seconds", BlueprintAutocast, DeprecatedFunction, DeprecationMessage = "FrameRateToInterval replaces this function, which returns the expected result of seconds per frame, rather than (incorrectly) frames per second."), Category = "Utilities|Time Management")
 	static float Conv_FrameRateToSeconds(const FFrameRate& InFrameRate);
 
+	/** Converts a FrameRate to an interval float representing the frame time in seconds ie: 1/30 returns 0.0333333 */
+	UFUNCTION(BlueprintPure, meta = (DisplayName = "FrameRate To Interval", BlueprintAutocast), Category = "Utilities|Time Management")
+	static float Conv_FrameRateToInterval(const FFrameRate InFrameRate);
+
 	/** Converts an QualifiedFrameTime to seconds. */
 	UFUNCTION(BlueprintPure, meta = (DisplayName = "QualifiedFrameTime To Seconds", BlueprintAutocast), Category = "Utilities|Time Management")
 	static float Conv_QualifiedFrameTimeToSeconds(const FQualifiedFrameTime& InFrameTime);
