Time Management: Fix EvaluateTimecode when frame rate is low

diff --git a/Engine/Source/Runtime/TimeManagement/Private/FrameNumberTimeEvaluator.cpp b/Engine/Source/Runtime/TimeManagement/Private/FrameNumberTimeEvaluator.cpp
--- a/Engine/Source/Runtime/TimeManagement/Private/FrameNumberTimeEvaluator.cpp
+++ b/Engine/Source/Runtime/TimeManagement/Private/FrameNumberTimeEvaluator.cpp
@@ -159,10 +159,13 @@ TValueOrError<FFrameTime, FExpressionError> FFrameNumberTimeEvaluator::EvaluateT
 
 	// Convert any excess frames into seconds
 	const int32 MaxFrame = (int32)FMath::RoundToInt(InDisplayFrameRate.AsDecimal());
-	while (Frames >= MaxFrame)
+	if (MaxFrame > 0)
 	{
-		Seconds += 1;
-		Frames -= MaxFrame;
+		while (Frames >= MaxFrame)
+		{
+			Seconds += 1;
+			Frames -= MaxFrame;
+		}
 	}
 
 	// Convert any excess seconds into minutes
