Attempt to fix MacArm64  editor compile error.

diff --git a/Engine/Source/Runtime/TimeManagement/Public/FrameBasedMusicMap.h b/Engine/Source/Runtime/TimeManagement/Public/FrameBasedMusicMap.h
--- a/Engine/Source/Runtime/TimeManagement/Public/FrameBasedMusicMap.h
+++ b/Engine/Source/Runtime/TimeManagement/Public/FrameBasedMusicMap.h
@@ -372,6 +372,11 @@ struct FMarkerProviderEntry
 {
 	FString Label;
 	FFrameNumber FrameNumber;
+
+	FMarkerProviderEntry(const FString& InLabel, FFrameNumber InFrameNumber)
+		: Label(InLabel)
+		, FrameNumber(InFrameNumber)
+	{}
 };
 
 USTRUCT()
@@ -383,6 +388,10 @@ struct FMarkerProviderResults
 	{
 		FString Name;
 		TArray<FMarkerProviderEntry> Markers;
+
+		FChannel(const FString& InName)
+			: Name(InName)
+		{}
 	};
 
 	bool IsEmpty() const { return Channels.IsEmpty(); }
