TimeManagement: Added catchup custom timestep.

This is a building block towards being able to absorb Unreal Editor hitches while take recording. It is a new custom timestep that does the following:

 * Control the Engine TimeStep via a fixed frame rate that catches up with real time.
    - Stays in sync with platform time.
    - Blocks to prevent getting ahead of real time.
    - Does not block when it needs to catch up
    - If it falls behind too much, it will increase simulation delta times.

This does not account for external live inputs, which another important piece for this initiative. In particular what is still missing is:

* A way to assign Timecode to frames that is also resistant to hitches. This is important because buffered data is looked up based on timecode.
* Buffer all live inputs that we care about. Livelink and video already support this.
* Putting the engine in this mode in an easy and intuitive way for Take Recorder usage.

diff --git a/Engine/Source/Runtime/TimeManagement/Private/CatchupFixedRateCustomTimeStep.cpp b/Engine/Source/Runtime/TimeManagement/Private/CatchupFixedRateCustomTimeStep.cpp
new file mode 100644
--- /dev/null
+++ b/Engine/Source/Runtime/TimeManagement/Private/CatchupFixedRateCustomTimeStep.cpp
@@ -0,0 +1,158 @@
+// Copyright Epic Games, Inc. All Rights Reserved.
+
+#include "CatchupFixedRateCustomTimeStep.h"
+#include "ITimeManagementModule.h"
+#include "Misc/App.h"
+
+
+#include UE_INLINE_GENERATED_CPP_BY_NAME(CatchupFixedRateCustomTimeStep)
+
+UCatchupFixedRateCustomTimeStep::UCatchupFixedRateCustomTimeStep(const FObjectInitializer& ObjectInitializer)
+	: Super(ObjectInitializer)
+	, FrameRate(24,1)
+{
+}
+
+bool UCatchupFixedRateCustomTimeStep::Initialize(UEngine* InEngine)
+{
+	// We begin our simulation all caught up with platform time.
+	SimulationSeconds = FPlatformTime::Seconds();
+
+	// But we quantize it to multiples of the delta time
+	const double DefaultDeltaSeconds = GetFixedFrameRate().AsInterval();
+	SimulationSeconds = FMath::RoundToDouble(SimulationSeconds / DefaultDeltaSeconds) * DefaultDeltaSeconds;
+
+	return true;
+}
+
+void UCatchupFixedRateCustomTimeStep::Shutdown(UEngine* InEngine)
+{
+	// Empty but implemented because it is PURE_VIRTUAL
+}
+
+bool UCatchupFixedRateCustomTimeStep::UpdateTimeStep(UEngine* InEngine)
+{
+	// Copy "CurrentTime" (used during the previous frame) into "LastTime"
+	UpdateApplicationLastTime();
+
+	// We will use the currently sampled platform time for all the timestep calculations.
+	const double CurrentPlatformSeconds = FPlatformTime::Seconds();
+
+	// Determine simulation delta seconds to apply to this frame
+	const double DeltaSeconds = CalculateDeltaSeconds(SimulationSeconds, CurrentPlatformSeconds);
+
+	// Increase the simulation time by this amount.
+	SimulationSeconds += DeltaSeconds;
+
+	// Eliminate accumulation errors, such that simulation time is always at an exact frame boundary.
+	const double DefaultDeltaSeconds = GetFixedFrameRate().AsInterval();
+	SimulationSeconds = FMath::RoundToDouble(SimulationSeconds / DefaultDeltaSeconds) * DefaultDeltaSeconds;
+
+	// Idle time is how much time we'll have to block, i.e. how much Simulation time is head of platform time.
+	FApp::SetIdleTime(FMath::Max(0.0, SimulationSeconds - CurrentPlatformSeconds));
+
+	// If the simulation is ahead, we should let platform time reach simulation time because
+	// simulation time determines live input sampling data, which cannot be available if we are
+	// simulating ahead of platform time since they would be in the future.
+	BlockUntilPlatformSeconds(SimulationSeconds);
+
+	// Current platform time should now be right after the desired SimulationSeconds, with an overshoot
+	FApp::SetIdleTimeOvershoot(FMath::Max(0.0, FPlatformTime::Seconds() - SimulationSeconds));
+
+	// Current time is always our simulation time, since that is the purpose of this custom timestep.
+	FApp::SetCurrentTime(SimulationSeconds);
+
+	// Delta time is our catchup delta time, which should normally be equal to the inverse of our frame rate.
+	FApp::SetDeltaTime(DeltaSeconds);
+
+	return false; // false means that the Engine's TimeStep should NOT be performed.
+}
+
+ECustomTimeStepSynchronizationState UCatchupFixedRateCustomTimeStep::GetSynchronizationState() const
+{
+	// If simulation is falling behind (or too far ahead), then consider the state as not fully synchronized.
+	if (FMath::Abs(FPlatformTime::Seconds() - SimulationSeconds) > MaxCatchupSeconds / 2)
+	{
+		return ECustomTimeStepSynchronizationState::Synchronizing;
+	}
+
+	return ECustomTimeStepSynchronizationState::Synchronized;
+}
+
+FFrameRate UCatchupFixedRateCustomTimeStep::GetFixedFrameRate() const
+{
+	return FrameRate;
+}
+
+void UCatchupFixedRateCustomTimeStep::BlockUntilPlatformSeconds(const double TargetPlatformSeconds) const
+{
+	const double IdleSeconds = TargetPlatformSeconds - FPlatformTime::Seconds();
+
+	// Early return if we're already there.
+	if (IdleSeconds <= 0.0)
+	{
+		return;
+	}
+
+	// Normal sleep for the bulk of the idle time.
+	
+	constexpr double EnoughTimeToWaitSleeping = 4e-3;
+
+	if (IdleSeconds > EnoughTimeToWaitSleeping)
+	{
+		constexpr double MarginToSpinSeconds = 2e-3;
+		FPlatformProcess::SleepNoStats(IdleSeconds - MarginToSpinSeconds);
+	}
+
+	// Give up timeslice for the small remainder of wait time.
+	while (FPlatformTime::Seconds() < TargetPlatformSeconds)
+	{
+		FPlatformProcess::SleepNoStats(0);
+	}
+}
+
+double UCatchupFixedRateCustomTimeStep::CalculateDeltaSeconds(const double CurrentSimulationSeconds, const double CurrentPlatformSeconds) const
+{
+	// We will adapt to systemic simulation fall-behind by increasing the simulation delta time
+	// * We define a maximum alloweable catchup time, beyond which we'll just catch up instantly instead.
+	// * If on the 2nd half, then we double the delta time so that it catches up faster.
+	// * Otherwise use the default delta time.
+
+	// Our default delta time is the inverse of our fixed rate.
+	const double DefaultDeltaSeconds = GetFixedFrameRate().AsInterval();
+
+	const double CatchupSeconds = CurrentPlatformSeconds - CurrentSimulationSeconds;
+
+	if (CatchupSeconds >= MaxCatchupSeconds)
+	{
+		// CatchupDeltas is how many default delta times our simulation is behind platform time.
+		const int32 CatchupDeltas = CatchupSeconds / DefaultDeltaSeconds;
+
+		const double ImmediateCatchupDeltaSeconds = CatchupDeltas * DefaultDeltaSeconds;
+
+		// Note: We don't expect this log to happen often because it only happens after catchup mechanisms fail to keep up.
+		UE_LOG(LogTimeManagement, Warning,
+			TEXT("CatchupFixedRateCustomTimeStep: Because the simulation fell behind the limit of %.1f seconds from platform time,"
+				" used a game delta time of %.1f to immediately catch up."), MaxCatchupSeconds, ImmediateCatchupDeltaSeconds);
+
+		return ImmediateCatchupDeltaSeconds;
+	}
+
+	const double GradualCatchupThresholdSeconds = MaxCatchupSeconds / 2;
+
+	if (CatchupSeconds >= GradualCatchupThresholdSeconds)
+	{
+		// Max value by which we are going to scale the default delta time to compensate for being so far behind.
+		constexpr double MaxScaleFactor = 4.0;
+
+		// Calculate the scale factor based on how far behind we are, ranging from 1.0 to MaxScaleFactor.
+		const double Dy = MaxScaleFactor - 1.0;
+		const double Dx = MaxCatchupSeconds - GradualCatchupThresholdSeconds;
+		const double ScaleFactor = 1.0 + (CatchupSeconds - GradualCatchupThresholdSeconds) * (Dy / Dx);
+
+		// Round the scale factor to keep it aligned with DefaultDeltaSeconds intervals.
+		return FMath::RoundToDouble(ScaleFactor) * DefaultDeltaSeconds;
+	}
+
+	return DefaultDeltaSeconds;
+}
diff --git a/Engine/Source/Runtime/TimeManagement/Private/TimeManagementModule.cpp b/Engine/Source/Runtime/TimeManagement/Private/TimeManagementModule.cpp
--- a/Engine/Source/Runtime/TimeManagement/Private/TimeManagementModule.cpp
+++ b/Engine/Source/Runtime/TimeManagement/Private/TimeManagementModule.cpp
@@ -4,6 +4,10 @@
 #include "TimedDataInputCollection.h"
 #include "CommonFrameRates.h"
 
+
+DEFINE_LOG_CATEGORY(LogTimeManagement)
+
+
 class FTimeManagementModule : public ITimeManagementModule
 {
 public:
diff --git a/Engine/Source/Runtime/TimeManagement/Public/CatchupFixedRateCustomTimeStep.h b/Engine/Source/Runtime/TimeManagement/Public/CatchupFixedRateCustomTimeStep.h
new file mode 100644
--- /dev/null
+++ b/Engine/Source/Runtime/TimeManagement/Public/CatchupFixedRateCustomTimeStep.h
@@ -0,0 +1,61 @@
+// Copyright Epic Games, Inc. All Rights Reserved.
+
+#pragma once
+
+#include "FixedFrameRateCustomTimeStep.h"
+
+#include "CatchupFixedRateCustomTimeStep.generated.h"
+
+class UEngine;
+class UObject;
+
+/**
+ * Control the Engine TimeStep via a fixed frame rate that catches up with real time.
+ * 
+ *   - Stays in sync with platform time.
+ *   - Blocks to prevent getting ahead of real time.
+ *   - Does not block when it needs to catch up
+ *   - If it falls behind too much, it will increase simulation delta times.
+ * 
+ */
+UCLASS(Blueprintable, editinlinenew, meta = (DisplayName = "Catchup Fixed Rate"), MinimalAPI)
+class UCatchupFixedRateCustomTimeStep : public UFixedFrameRateCustomTimeStep
+{
+	GENERATED_UCLASS_BODY()
+
+public:
+	//~ Begin UFixedFrameRateCustomTimeStep interface
+	TIMEMANAGEMENT_API virtual bool Initialize(UEngine* InEngine) override;
+	TIMEMANAGEMENT_API virtual void Shutdown(UEngine* InEngine) override;
+	TIMEMANAGEMENT_API virtual bool UpdateTimeStep(UEngine* InEngine) override;
+	TIMEMANAGEMENT_API virtual ECustomTimeStepSynchronizationState GetSynchronizationState() const override;
+	TIMEMANAGEMENT_API virtual FFrameRate GetFixedFrameRate() const override;
+	//~ End UFixedFrameRateCustomTimeStep interface
+
+private:
+
+	/** Blocks until platform time reaches the given time */
+	void BlockUntilPlatformSeconds(const double TargetPlatformSeconds) const;
+
+	/**
+	 * Calculates the delta time to apply for the current frame based on the difference between simulation time and platform time.
+	 * Adjusts the simulation delta time to account for systemic fall-behind, allowing the simulation to catch up smoothly.
+	 */
+	double CalculateDeltaSeconds(const double CurrentSimulationSeconds, const double CurrentPlatformSeconds) const;
+
+public:
+
+	/** Desired simulation frame rate */
+	UPROPERTY(EditAnywhere, Category = "Timing")
+	FFrameRate FrameRate;
+
+	/** Maximum catchup time in seconds. Simulation will catch up instantly if it falls behind beyond this time with respect to platform time */
+	UPROPERTY(EditAnywhere, Category = "Timing", meta = (ClampMin = "2"))
+	double MaxCatchupSeconds = 16;
+
+private:
+
+	/* Keeps track of our simulation time, which is intentionally kept close to platform time */
+	double SimulationSeconds = 0;
+
+};
diff --git a/Engine/Source/Runtime/TimeManagement/Public/ITimeManagementModule.h b/Engine/Source/Runtime/TimeManagement/Public/ITimeManagementModule.h
--- a/Engine/Source/Runtime/TimeManagement/Public/ITimeManagementModule.h
+++ b/Engine/Source/Runtime/TimeManagement/Public/ITimeManagementModule.h
@@ -9,6 +9,8 @@
 
 #define TIMEMANAGEMENT_MODULE_NAME TEXT("TimeManagement")
 
+DECLARE_LOG_CATEGORY_EXTERN(LogTimeManagement, Log, All);
+
 class FTimedDataInputCollection;
 
 class ITimeManagementModule : public IModuleInterface
