Fix timecode not displaying correctly when the Windows text scale is not 100%
LiveLinkHub - Change timecode color based on configuration status

#rnx

diff --git a/Engine/Plugins/Animation/LiveLinkHub/Source/LiveLinkHub/Private/UI/Widgets/SLiveLinkTimecode.cpp b/Engine/Plugins/Animation/LiveLinkHub/Source/LiveLinkHub/Private/UI/Widgets/SLiveLinkTimecode.cpp
--- a/Engine/Plugins/Animation/LiveLinkHub/Source/LiveLinkHub/Private/UI/Widgets/SLiveLinkTimecode.cpp
+++ b/Engine/Plugins/Animation/LiveLinkHub/Source/LiveLinkHub/Private/UI/Widgets/SLiveLinkTimecode.cpp
@@ -6,6 +6,7 @@
 #include "Framework/MultiBox/MultiBoxBuilder.h"
 #include "ILiveLinkModule.h"
 #include "Internationalization/Text.h"
+#include "ILiveLinkModule.h"
 #include "ISettingsModule.h"
 #include "Misc/App.h"
 #include "Modules/ModuleManager.h"
@@ -25,6 +26,8 @@ FSlateColor SLiveLinkTimecode::GetIconColor() const
 {
 	const ULiveLinkHubTimeAndSyncSettings* TimeAndSyncSettinsgs = GetDefault<ULiveLinkHubTimeAndSyncSettings>();
 	const bool bCustomTimeStepEnabled = TimeAndSyncSettinsgs->bUseLiveLinkHubAsCustomTimeStepSource;
+	static const FSlateColor ValidColor = ILiveLinkModule::Get().GetStyle()->GetSlateColor("LiveLink.Color.Valid");
+	static FSlateColor ErrorColor = ILiveLinkModule::Get().GetStyle()->GetSlateColor("LiveLink.Color.Error");
 
 	TSharedPtr<FSlateStyleSet> Style = ILiveLinkModule::Get().GetStyle();
 	if (bCustomTimeStepEnabled)
diff --git a/Engine/Source/Runtime/TimeManagement/Private/Widgets/STimecode.cpp b/Engine/Source/Runtime/TimeManagement/Private/Widgets/STimecode.cpp
--- a/Engine/Source/Runtime/TimeManagement/Private/Widgets/STimecode.cpp
+++ b/Engine/Source/Runtime/TimeManagement/Private/Widgets/STimecode.cpp
@@ -4,6 +4,8 @@
 
 #include "Fonts/FontMeasure.h"
 #include "Framework/Application/SlateApplication.h"
+#include "Framework/Text/PlainTextLayoutMarshaller.h"
+#include "Internationalization/IBreakIterator.h"
 #include "Rendering/DrawElements.h"
 #include "Rendering/RenderingCommon.h"
 
@@ -37,6 +39,8 @@ void STimecode::Construct(const FArguments& InArgs)
 	bDisplayLabel = InArgs._DisplayLabel;
 	LabelFont = InArgs._LabelFont;
 	LabelColor = InArgs._LabelColor;
+
+	TextLayoutCache = MakeUnique<FSlateTextBlockLayout>(this, FTextBlockStyle::GetDefault(), ETextShapingMethod::Auto, ETextFlowDirection::LeftToRight, FCreateSlateTextLayout(), FPlainTextLayoutMarshaller::Create(), nullptr);
 }
 
 int32 STimecode::OnPaint(const FPaintArgs& Args, const FGeometry& AllottedGeometry, const FSlateRect& MyCullingRect, FSlateWindowElementList& OutDrawElements, int32 LayerId, const FWidgetStyle& InWidgetStyle, bool bParentEnabled) const
@@ -111,18 +115,39 @@ int32 STimecode::OnPaint(const FPaintArgs& Args, const FGeometry& AllottedGeomet
 	return NewLayerId;
 }
 
-FVector2D STimecode::ComputeDesiredSize(float) const
+FVector2D STimecode::ComputeDesiredSize(float LayoutScale) const
 {
 	FString TimecodeString = Timecode.Get().ToString();
-	const TSharedRef< FSlateFontMeasure >& FontMeasureService = FSlateApplication::Get().GetRenderer()->GetFontMeasureService();
-	const FVector2D TimecodeSize = FontMeasureService->Measure(TimecodeString, TimecodeFont.Get());
 
 	if (bDisplayLabel.Get())
 	{
+		const TSharedRef< FSlateFontMeasure >& FontMeasureService = FSlateApplication::Get().GetRenderer()->GetFontMeasureService();
+		const FVector2D TimecodeSize = FontMeasureService->Measure(TimecodeString, TimecodeFont.Get());
 		const FVector2D LabelSize = bDisplayLabel.Get() ? FontMeasureService->Measure(TEXT("HR"), LabelFont.Get()) : FVector2D::ZeroVector;
 		return FVector2D(TimecodeSize.X, TimecodeSize.Y + LabelSize.Y + TimecodeImpl::LabelOffsetY);
 	}
-	return TimecodeSize;
+	else
+	{
+		// If the label is not displayed, then compute the size correctly with the layout scale.
+		FTextBlockStyle Style;
+		Style.SetFont(TimecodeFont.Get());
+
+		const FVector2D DesiredHorizontalTextSize = TextLayoutCache->ComputeDesiredSize(
+			FSlateTextBlockLayout::FWidgetDesiredSizeArgs(
+				FText::FromString(TimecodeString),
+				FText(),
+				0.f,
+				false,
+				ETextWrappingPolicy::DefaultWrapping,
+				ETextTransformPolicy::None,
+				FMargin(),
+				1.f,
+				true,
+				ETextJustify::Left),
+			LayoutScale, Style);
+
+		return DesiredHorizontalTextSize;
+	}
 }
 
 bool STimecode::ComputeVolatility() const
diff --git a/Engine/Source/Runtime/TimeManagement/Public/Widgets/STimecode.h b/Engine/Source/Runtime/TimeManagement/Public/Widgets/STimecode.h
--- a/Engine/Source/Runtime/TimeManagement/Public/Widgets/STimecode.h
+++ b/Engine/Source/Runtime/TimeManagement/Public/Widgets/STimecode.h
@@ -11,6 +11,7 @@
 #include "Misc/Timecode.h"
 #include "Styling/CoreStyle.h"
 #include "Styling/SlateColor.h"
+#include "Widgets/Text/SlateTextBlockLayout.h"
 
 class STimecode : public SLeafWidget
 {
@@ -52,7 +53,7 @@ public:
 protected:
 	// SWidget overrides
 	TIMEMANAGEMENT_API virtual int32 OnPaint(const FPaintArgs& Args, const FGeometry& AllottedGeometry, const FSlateRect& MyCullingRect, FSlateWindowElementList& OutDrawElements, int32 LayerId, const FWidgetStyle& InWidgetStyle, bool bParentEnabled) const override;
-	TIMEMANAGEMENT_API virtual FVector2D ComputeDesiredSize(float) const override;
+	TIMEMANAGEMENT_API virtual FVector2D ComputeDesiredSize(float LayoutScale) const override;
 	TIMEMANAGEMENT_API virtual bool ComputeVolatility() const override;
 
 private:
@@ -63,4 +64,7 @@ private:
 	TAttribute<bool> bDisplayLabel;
 	TAttribute<FSlateFontInfo> LabelFont;
 	TAttribute<FSlateColor> LabelColor;
+
+	/** Test layout cache used to correctly compute the text size for the timecode text. */
+	TUniquePtr<FSlateTextBlockLayout> TextLayoutCache;
 };
