LiveLinkHub - Don't display subframes in the hub

#rnx

diff --git a/Engine/Plugins/Animation/LiveLinkHub/Source/LiveLinkHub/Private/UI/Widgets/SLiveLinkTimecode.cpp b/Engine/Plugins/Animation/LiveLinkHub/Source/LiveLinkHub/Private/UI/Widgets/SLiveLinkTimecode.cpp
--- a/Engine/Plugins/Animation/LiveLinkHub/Source/LiveLinkHub/Private/UI/Widgets/SLiveLinkTimecode.cpp
+++ b/Engine/Plugins/Animation/LiveLinkHub/Source/LiveLinkHub/Private/UI/Widgets/SLiveLinkTimecode.cpp
@@ -115,6 +115,7 @@ void SLiveLinkTimecode::Construct(const FArguments& InArgs)
 			{
 				return FApp::GetTimecode();
 			}))
+			.bDisplaySubframes(false)
 		]
 		+ SHorizontalBox::Slot()
 		.AutoWidth()
diff --git a/Engine/Source/Runtime/TimeManagement/Private/Widgets/STimecode.cpp b/Engine/Source/Runtime/TimeManagement/Private/Widgets/STimecode.cpp
--- a/Engine/Source/Runtime/TimeManagement/Private/Widgets/STimecode.cpp
+++ b/Engine/Source/Runtime/TimeManagement/Private/Widgets/STimecode.cpp
@@ -39,6 +39,7 @@ void STimecode::Construct(const FArguments& InArgs)
 	bDisplayLabel = InArgs._DisplayLabel;
 	LabelFont = InArgs._LabelFont;
 	LabelColor = InArgs._LabelColor;
+	bDisplaySubframes = InArgs._bDisplaySubframes;
 
 	TextLayoutCache = MakeUnique<FSlateTextBlockLayout>(this, FTextBlockStyle::GetDefault(), ETextShapingMethod::Auto, ETextFlowDirection::LeftToRight, FCreateSlateTextLayout(), FPlainTextLayoutMarshaller::Create(), nullptr);
 }
@@ -49,6 +50,10 @@ int32 STimecode::OnPaint(const FPaintArgs& Args, const FGeometry& AllottedGeomet
 
 	FTimecode TimecodeToPaint = Timecode.Get();
 	FString TimecodeToPaintString = TimecodeToPaint.ToString();
+	if (!bDisplaySubframes && TimecodeToPaint.Subframe != 0)
+	{
+		TimecodeToPaintString.LeftChopInline(3);
+	}
 	const FSlateFontInfo& TimecodeFontInfo = TimecodeFont.Get();
 
 	int32 NewLayerId = LayerId+1;
@@ -117,7 +122,12 @@ int32 STimecode::OnPaint(const FPaintArgs& Args, const FGeometry& AllottedGeomet
 
 FVector2D STimecode::ComputeDesiredSize(float LayoutScale) const
 {
+	const FTimecode TimecodeToPaint = Timecode.Get();
 	FString TimecodeString = Timecode.Get().ToString();
+	if (!bDisplaySubframes && TimecodeToPaint.Subframe != 0)
+	{
+		TimecodeString.LeftChopInline(3);
+	}
 
 	if (bDisplayLabel.Get())
 	{
diff --git a/Engine/Source/Runtime/TimeManagement/Public/Widgets/STimecode.h b/Engine/Source/Runtime/TimeManagement/Public/Widgets/STimecode.h
--- a/Engine/Source/Runtime/TimeManagement/Public/Widgets/STimecode.h
+++ b/Engine/Source/Runtime/TimeManagement/Public/Widgets/STimecode.h
@@ -32,13 +32,14 @@ public:
 		SLATE_ATTRIBUTE(FSlateFontInfo, TimecodeFont)
 		/** The color for the timecode text */
 		SLATE_ATTRIBUTE(FSlateColor, TimecodeColor)
-
 		/** Should display the label (hours, mins, secs, frames) */
 		SLATE_ATTRIBUTE(bool, DisplayLabel)
 		/** The font for this label text */
 		SLATE_ATTRIBUTE(FSlateFontInfo, LabelFont)
 		/** The color for this label text */
 		SLATE_ATTRIBUTE(FSlateColor, LabelColor)
+		/** Whether to display subframes. */
+		SLATE_ARGUMENT(bool, bDisplaySubframes)
 	SLATE_END_ARGS()
 
 	TIMEMANAGEMENT_API STimecode();
@@ -65,6 +66,9 @@ private:
 	TAttribute<FSlateFontInfo> LabelFont;
 	TAttribute<FSlateColor> LabelColor;
 
+	/** Whether to display subframes. */
+	bool bDisplaySubframes = true;
+
 	/** Test layout cache used to correctly compute the text size for the timecode text. */
 	TUniquePtr<FSlateTextBlockLayout> TextLayoutCache;
 };
