Implement time code estimation based on linear regression.

In combination with a fixed time step, this can be used to address the engine hitching on certain frames. Until now, the frame after a very long frame would use the platform time to determine the currentl time code, when really it should use the FApp::CurrentTime (game time set by UCustomTimeStep).

The time code estimator wraps a UTimeCodeProvider and tries to correleate underlying clock time of a UCustomTimeStep with time code.
Every time UTimeCodeProvider::FetchAndUpdate is called, we associate the real underlying clock time with the time code (e.g. platform time).
In GetQualifiedFrameTime, which is when we evaluate what the current time code is, we perform linear regression based on the current game time (the accumulated delta times), set by the UCustomTimeStep.
This way, when the engine hitches, it will sample the right time code in subsequent frames.

Introduces Timecode.LogTimecodeSampling and Timecode.LogTimecodeEstimation for debugging sampled and estimated frames.

diff --git a/Engine/Source/Runtime/TimeManagement/Private/CatchupFixedRateCustomTimeStep.cpp b/Engine/Source/Runtime/TimeManagement/Private/CatchupFixedRateCustomTimeStep.cpp
--- a/Engine/Source/Runtime/TimeManagement/Private/CatchupFixedRateCustomTimeStep.cpp
+++ b/Engine/Source/Runtime/TimeManagement/Private/CatchupFixedRateCustomTimeStep.cpp
@@ -64,7 +64,6 @@ bool UCatchupFixedRateCustomTimeStep::UpdateTimeStep(UEngine* InEngine)
 
 	// Delta time is our catchup delta time, which should normally be equal to the inverse of our frame rate.
 	FApp::SetDeltaTime(DeltaSeconds);
-
 	return false; // false means that the Engine's TimeStep should NOT be performed.
 }
 
@@ -84,6 +83,11 @@ FFrameRate UCatchupFixedRateCustomTimeStep::GetFixedFrameRate() const
 	return FrameRate;
 }
 
+TOptional<double> UCatchupFixedRateCustomTimeStep::GetUnderlyingClockTime_AnyThread()
+{
+	return FPlatformTime::Seconds();
+}
+
 void UCatchupFixedRateCustomTimeStep::BlockUntilPlatformSeconds(const double TargetPlatformSeconds) const
 {
 	const double IdleSeconds = TargetPlatformSeconds - FPlatformTime::Seconds();
diff --git a/Engine/Source/Runtime/TimeManagement/Private/Estimation/TimecodeEstimator.cpp b/Engine/Source/Runtime/TimeManagement/Private/Estimation/TimecodeEstimator.cpp
new file mode 100644
--- /dev/null
+++ b/Engine/Source/Runtime/TimeManagement/Private/Estimation/TimecodeEstimator.cpp
@@ -0,0 +1,177 @@
+﻿// Copyright Epic Games, Inc. All Rights Reserved.
+
+#include "Estimation/TimecodeEstimator.h"
+
+#include "Engine/EngineCustomTimeStep.h"
+#include "Engine/TimecodeProvider.h"
+#include "Estimation/IClockedTimeStep.h"
+#include "HAL/IConsoleManager.h"
+#include "ITimeManagementModule.h"
+
+namespace UE::TimeManagement::TimecodeEstimation
+{
+static TAutoConsoleVariable<bool> CVarLogSampling(
+	TEXT("Timecode.LogTimecodeSampling"), false, TEXT("When estimating timecode, whether to log sampled time and the current time. For this to take effect, you must use UTimecodeRegressionProvider as custom engine timestep.")
+	);
+static TAutoConsoleVariable<bool> CVarLogEstimation(
+	TEXT("Timecode.LogTimecodeEstimation"), false, TEXT("When estimating timecode, whether to log estimated time and the current time. For this to take effect, you must use UTimecodeRegressionProvider as custom engine timestep.")
+);
+static TAutoConsoleVariable<bool> CVarLogTimecodeDifference(
+	TEXT("Timecode.LogTimecodeDifference"), false, TEXT("Logs the timecode difference between the timecode of the current underlying clock and what is estimated. This is useful for debugging.")
+);
+static TAutoConsoleVariable<float> CVarUnclearEstimationSubframeThreshold(
+	TEXT("Timecode.UnclearEstimationSubframeThreshold"), 0.1f,
+	TEXT("If Abs(timecode's subframe  - 0.5f) <= this value, a warning is logged. That is because an estimated value like 14:11:43:16.49 is not a "
+		"clear estimate... it is very close to both 14:11:43:16.00, or 14:11:43:17.00")
+);
+
+static void LogEstimatedTime(
+	double InRelativeTime, IClockedTimeStep& Clock, UTimecodeProvider& TimecodeProvider,
+	const FFrameTime& InEstimatedTime, const FFrameTime& InUnroundedEstimatedTime, const FFrameRate& InFrameRate
+	)
+{
+	const auto ToString = [](const FTimecode& InTime)
+	{
+		constexpr bool bForceSignDisplay = false, bDisplaySubframe = true;
+		return InTime.ToString(bForceSignDisplay, bDisplaySubframe);
+	};
+	
+	const TOptional<double> ClockTime = Clock.GetUnderlyingClockTime_AnyThread();
+	const double AppTime = FApp::GetCurrentTime();
+	UE_CLOG(CVarLogEstimation.GetValueOnAnyThread(), LogTimeManagement, Log,
+		TEXT("Estimate %f at %s \t\t(Unrounded: %s \tClock %s, \tApp: %f)"),
+		InRelativeTime,
+		*ToString(FQualifiedFrameTime(InEstimatedTime, InFrameRate).ToTimecode()),
+		*ToString(FQualifiedFrameTime(InUnroundedEstimatedTime, InFrameRate).ToTimecode()),
+		ClockTime ? *FString::SanitizeFloat(*ClockTime) : TEXT("unset"), AppTime
+		);
+
+	if (CVarLogTimecodeDifference.GetValueOnAnyThread())
+	{
+		TimecodeProvider.FetchAndUpdate();
+		const FQualifiedFrameTime ActualFrameTime = TimecodeProvider.GetQualifiedFrameTime();
+		
+		FTimecode ActualTC = ActualFrameTime.ToTimecode();
+		FTimecode EstimatedTC = FQualifiedFrameTime(InEstimatedTime, InFrameRate).ToTimecode();
+		FTimecode EstimatedRoundedTC = FQualifiedFrameTime(InUnroundedEstimatedTime, InFrameRate).ToTimecode();
+		const bool bIsTimeEqual = EstimatedTC == ActualTC;
+		if (bIsTimeEqual)
+		{
+			return;
+		}
+
+		if (InEstimatedTime > ActualFrameTime.Time)
+		{
+			const FTimecode AbsDeltaTC = FQualifiedFrameTime(InEstimatedTime - ActualFrameTime.Time, InFrameRate).ToTimecode();
+			UE_CLOG(!bIsTimeEqual && InEstimatedTime > ActualFrameTime.Time, LogTimeManagement, Warning,
+				TEXT("Leading timecode \tDelta: +%s \tActual: %s \tEstimate: %s \tEstimate (unrounded): %s"),
+				*ToString(AbsDeltaTC), *ToString(ActualTC), *ToString(EstimatedTC), *ToString(EstimatedRoundedTC)
+			);
+		}
+		else
+		{
+			const FTimecode AbsDeltaTC = FQualifiedFrameTime(ActualFrameTime.Time - InEstimatedTime, InFrameRate).ToTimecode();
+			UE_CLOG(!bIsTimeEqual && InEstimatedTime < ActualFrameTime.Time, LogTimeManagement, Warning,
+				TEXT("Trailing timecode \tDelta: -%s \tActual: %s \tEstimate: %s \tEstimate (unrounded): %s"),
+				*ToString(AbsDeltaTC), *ToString(ActualTC), *ToString(EstimatedTC), *ToString(EstimatedRoundedTC)
+				);
+		}
+	}
+}
+
+FTimecodeEstimator::FTimecodeEstimator(
+	SIZE_T InNumSamples,
+	UTimecodeProvider& InTimecode, IClockedTimeStep& InEngineCustomTimeStep
+	)
+	// Counter-intuitively, we should NOT initialize the start times because it's too early... defer until the data actually starts being sampled.
+	// For example, if the custom time step was just changed, then FApp::CurrentTime may not contain the correct value, yet.
+	// Or the API user might construct now and only use FTimecodeEstimator much later.
+	: StartClockTime()
+	, TimecodeProvider(InTimecode)
+	, EngineCustomTimeStep(InEngineCustomTimeStep)
+	, ClockToTimecodeSamples(InNumSamples)
+	, LastFrameRate(InTimecode.GetFrameRate())
+{
+	// There's no point in constructing FTimecodeEstimator if the number of linear regression samples is 0; it'll just use the latest value.
+	ensure(InNumSamples > 0); 
+}
+
+void FTimecodeEstimator::FetchAndUpdate()
+{
+	const TOptional<double> ClockTime = EngineCustomTimeStep.GetUnderlyingClockTime_AnyThread();
+	if (!ClockTime)
+	{
+		return;
+	}
+
+	if (!StartClockTime)
+	{
+		StartClockTime = ClockTime;
+	}
+	
+	TimecodeProvider.FetchAndUpdate(); // FetchAndUpdate fetches the latest timecode value so below GetQualifiedFrameTime returns the lastest value.
+	const FQualifiedFrameTime CurrentFrameTime = TimecodeProvider.GetQualifiedFrameTime();
+	const FFrameRate CurrentFrameRate = CurrentFrameTime.Rate;
+	
+	// In a true production environment, the frame rate of the timecode device should not really change on the fly, but we should handle it anyway.
+	if (CurrentFrameRate != LastFrameRate) 
+	{
+		ClockToTimecodeSamples = FCachedLinearRegressionSums(ClockToTimecodeSamples.Samples.Capacity());
+		LastFrameRate = CurrentFrameRate;
+	}
+
+	// We regress based on relative time for numerical stability. See StartClockTime docstring.
+	// Clock values can be very big but double precision is best near 0.
+	const FFrameTime& FrameTime = CurrentFrameTime.Time;
+	const double FrameTimeAsSeconds = LastFrameRate.AsSeconds(FrameTime);
+	const double RelativeTime = *ClockTime - *StartClockTime;
+
+	AddSampleAndUpdateSums(FVector2d{ RelativeTime, FrameTimeAsSeconds }, ClockToTimecodeSamples);
+	ComputeLinearRegressionSlopeAndOffset(ClockToTimecodeSamples.CachedSums, LinearRegressionFunction);
+
+	constexpr bool bForceSignDisplay = false, bDisplaySubframe = true;
+	UE_CLOG(CVarLogSampling.GetValueOnAnyThread(), LogTimeManagement, Log,
+		TEXT("Sampling %f at %s\t\t(Clock: %f, \tApp: %f)"),
+		RelativeTime, *FQualifiedFrameTime(FrameTime, LastFrameRate).ToTimecode().ToString(bForceSignDisplay, bDisplaySubframe),
+		*ClockTime, FApp::GetCurrentTime()
+		);
+}
+
+FQualifiedFrameTime FTimecodeEstimator::EstimateFrameTime() const
+{
+	if (!ClockToTimecodeSamples.IsEmpty()
+		&& ensureMsgf(StartClockTime, TEXT("Invariant: StartClockTime was supposed to have been set when the data was sampled!")))
+	{
+		const double RelativeTime = FApp::GetCurrentTime() - *StartClockTime;
+		const double EstimatedSeconds = LinearRegressionFunction.Evaluate(RelativeTime); 
+		const FFrameTime UnroundedEstimatedTime = EstimatedSeconds * LastFrameRate;
+		
+		// Subframe close to 0.5f? It may round to the wrong frame. Flag it so UE developer can investigate it, e.g. 0.47f is quite close 0.5f.
+		// We'd expect it to very close to the full frame, e.g. 14:11:43:16.04, or 14:11:43:16.87.
+		// A value like 14:11:43:16.42 is much closer to 0.5 than we'd expect... it may indicate that the time step is not linearly correlated
+		// to the timecode provider used (e.g. not really a fixed frame rate or some kind of noise).
+		const bool bUnclearEstimation = FMath::Abs(UnroundedEstimatedTime.GetSubFrame() - 0.5f) <= CVarUnclearEstimationSubframeThreshold.GetValueOnAnyThread();
+		constexpr bool bForceSignDisplay = false, bDisplaySubframe = true;
+		UE_CLOG(bUnclearEstimation, LogTimeManagement, Warning,
+			TEXT("Time %f estimate of %s is very close to 0.5 subframe border... the resulting timecode's frame may be off by 1."),
+			RelativeTime, *FQualifiedFrameTime(UnroundedEstimatedTime, LastFrameRate).ToTimecode().ToString(bForceSignDisplay, bDisplaySubframe)
+			);
+		
+		// FApp::GetCurrentTime() is usually slightly behind, e.g. relative underlying clock may be 2.028365, but FApp::GetCurrentTime() is 2.028364.
+		// That can cause slight trailing (provider 14:11:43:17.00, est. 14:11:43:16.99), or leading (provider: 14:11:43:19.00, est. 14:11:43:19.01).
+		// We only care about full frames... we drop subframes. That's why slight leading is no problem as the frame number stays the same.
+		// Trailing *is* a problem though because the frame is one less.
+		// Note: this comment was written when UCatchupFixedRateCustomTimeStep was the only timestep that UTimecodeRegressionProvider was set up with:
+		// so this observation is untested with other time steps that may have been added later (5.7+). If you find that behaviour is different
+		// between time steps, then we need to adjust this strategy.
+		const FFrameTime RoundedEstimatedTime = UnroundedEstimatedTime.RoundToFrame();
+		
+		LogEstimatedTime(RelativeTime, EngineCustomTimeStep, TimecodeProvider, RoundedEstimatedTime, UnroundedEstimatedTime, LastFrameRate);
+		return FQualifiedFrameTime(RoundedEstimatedTime, LastFrameRate);
+	}
+	
+	// This may cause jumps at the beginning, but it can be circumvented by warming up the engine, i.e. just let it run for a few frames
+	UE_LOG(LogTimeManagement, Log, TEXT("No data sampled, yet. This frame will fall back to actual timecode without estimation."));
+	return TimecodeProvider.GetQualifiedFrameTime();
+}
+}
diff --git a/Engine/Source/Runtime/TimeManagement/Private/Estimation/TimecodeRegressionProvider.cpp b/Engine/Source/Runtime/TimeManagement/Private/Estimation/TimecodeRegressionProvider.cpp
new file mode 100644
--- /dev/null
+++ b/Engine/Source/Runtime/TimeManagement/Private/Estimation/TimecodeRegressionProvider.cpp
@@ -0,0 +1,102 @@
+﻿// Copyright Epic Games, Inc. All Rights Reserved.
+
+#include "Estimation/TimecodeRegressionProvider.h"
+
+#include "Engine/Engine.h"
+#include "Engine/EngineCustomTimeStep.h"
+#include "Estimation/IClockedTimeStep.h"
+#include "ITimeManagementModule.h"
+#include "Misc/App.h"
+
+bool UTimecodeRegressionProvider::SupportsTimestep(UEngineCustomTimeStep* InTimestep)
+{
+	return InTimestep && InTimestep->Implements<UClockedTimeStep>();
+}
+
+void UTimecodeRegressionProvider::SetTimecodeImplementation(UTimecodeProvider* InTimecodeImpl)
+{
+	if (InTimecodeImpl && ensureMsgf(!IsInitialized(), TEXT("Invalid to call while in use by engine (i.e. Initialize has been called already)")))
+	{
+		TimecodeImpl = InTimecodeImpl;
+	}
+}
+
+void UTimecodeRegressionProvider::FetchAndUpdate()
+{
+	if (TimecodeEstimator)
+	{
+		TimecodeEstimator->FetchAndUpdate();
+	}
+	// This is just a sensible fallback if ReinitRegressionBuffer failed due to unsupported timestep class
+	else if (TimecodeImpl)
+	{
+		TimecodeImpl->FetchAndUpdate();
+	}
+}
+
+FQualifiedFrameTime UTimecodeRegressionProvider::GetQualifiedFrameTime() const
+{
+	const FQualifiedFrameTime Time = TimecodeEstimator
+		? TimecodeEstimator->EstimateFrameTime()
+		// This is just a sensible fallback if ReinitRegressionBuffer failed due to unsupported timestep class or missing call to SetTimecodeImplementation. 
+		: TimecodeImpl ? TimecodeImpl->GetQualifiedFrameTime() : FQualifiedFrameTime();
+	return Time;
+}
+
+ETimecodeProviderSynchronizationState UTimecodeRegressionProvider::GetSynchronizationState() const
+{
+	if (!IsInitialized())
+	{
+		return ETimecodeProviderSynchronizationState::Closed;
+	}
+	
+	// The custom time step can be changed after we've been initialized
+	if (!OwningEngine->GetCustomTimeStep())
+	{
+		return ETimecodeProviderSynchronizationState::Error;
+	}
+	
+	return TimecodeImpl->GetSynchronizationState();
+}
+
+bool UTimecodeRegressionProvider::Initialize(UEngine* InEngine)
+{
+	if (ensureMsgf(TimecodeImpl, TEXT("You forgot to call SetTimecodeImplementation before"))
+		&& TimecodeImpl->Initialize(InEngine))
+	{
+		OwningEngine = InEngine;
+		OwningEngine->OnCustomTimeStepChanged().AddUObject(this, &UTimecodeRegressionProvider::OnCustomTimeStepChanged);
+		ReinitRegressionBuffer();
+		return true;
+	}
+	return false;
+}
+
+void UTimecodeRegressionProvider::Shutdown(UEngine* InEngine)
+{
+	TimecodeImpl->Shutdown(InEngine);
+
+	OwningEngine->OnCustomTimeStepChanged().RemoveAll(this);
+	OwningEngine = nullptr;
+}
+
+void UTimecodeRegressionProvider::ReinitRegressionBuffer()
+{
+	check(IsInitialized());
+
+	UEngineCustomTimeStep* Timestep = OwningEngine->GetCustomTimeStep();
+	IClockedTimeStep* ClockedTimeStep = Timestep ? Cast<IClockedTimeStep>(Timestep) : nullptr;
+	UE_CLOG(Timestep && !ClockedTimeStep, LogTimeManagement, Error,
+		TEXT("Current timestep %s must implement the IClockedTimeStep in order to be used with UTimecodeRegressionProvider"),
+		*Timestep->GetName()
+		);
+	
+	if (ClockedTimeStep)
+	{
+		TimecodeEstimator.Emplace(FMath::Clamp(NumSampledFrames, 1, 65536), *TimecodeImpl, *ClockedTimeStep);
+	}
+	else
+	{
+		TimecodeEstimator.Reset();
+	}
+}
diff --git a/Engine/Source/Runtime/TimeManagement/Private/Misc/CachedLinearRegressionSums.cpp b/Engine/Source/Runtime/TimeManagement/Private/Misc/CachedLinearRegressionSums.cpp
new file mode 100644
--- /dev/null
+++ b/Engine/Source/Runtime/TimeManagement/Private/Misc/CachedLinearRegressionSums.cpp
@@ -0,0 +1,29 @@
+﻿// Copyright Epic Games, Inc. All Rights Reserved.
+
+#include "Misc/CachedLinearRegressionSums.h"
+
+namespace UE::TimeManagement
+{
+void AddSampleAndUpdateSums(const FVector2d& InSample, TModuloCircularBuffer<FVector2d>& InSamples, FLinearRegressionArgs& InCachedArguments)
+{
+	if (const FVector2d* OldestSample = InSamples.GetNextReplacedItem())
+	{
+		InCachedArguments.SumX -= OldestSample->X;
+		InCachedArguments.SumY -= OldestSample->Y;
+		InCachedArguments.SumXxY -= OldestSample->X * OldestSample->Y;
+		InCachedArguments.SumOfSquaredXes -= FMath::Square(OldestSample->X);
+		// Do this after updating the sums, as OldestSample is overwritten when we call Add
+		InSamples.Add(InSample); 
+	}
+	else
+	{
+		++InCachedArguments.Num;
+		InSamples.Add(InSample);
+	}
+
+	InCachedArguments.SumX += InSample.X;
+	InCachedArguments.SumY += InSample.Y;
+	InCachedArguments.SumXxY += InSample.X * InSample.Y;
+	InCachedArguments.SumOfSquaredXes += FMath::Square(InSample.X);
+}
+}
diff --git a/Engine/Source/Runtime/TimeManagement/Private/Misc/LinearRegression.cpp b/Engine/Source/Runtime/TimeManagement/Private/Misc/LinearRegression.cpp
new file mode 100644
--- /dev/null
+++ b/Engine/Source/Runtime/TimeManagement/Private/Misc/LinearRegression.cpp
@@ -0,0 +1,48 @@
+﻿// Copyright Epic Games, Inc. All Rights Reserved.
+
+#include "Misc/LinearRegression.h"
+
+namespace UE::TimeManagement
+{
+bool ComputeLinearRegressionSlopeAndOffset(
+	const FLinearRegressionArgs& InArgs,
+	double& OutSlope, double& OutOffset
+	)
+{
+	// This happens when all data points are equal. The independent variables are all on the same point. There is no slope.
+	const double Variance = InArgs.SumOfSquaredXes - (FMath::Square(InArgs.SumX) / InArgs.Num);
+	if (FMath::IsNearlyZero(Variance))
+	{
+		OutSlope = 0.0;
+		OutOffset = 0.0;
+		return false;
+	}
+
+	// Refresher: https://www.ncl.ac.uk/webtemplate/ask-assets/external/maths-resources/statistics/regression-and-correlation/simple-linear-regression.html (3rd July, 2025)
+	const double TargetX = InArgs.SumX / InArgs.Num;
+	const double TargetY = InArgs.SumY / InArgs.Num;
+	OutSlope = (InArgs.SumXxY - (InArgs.SumX * InArgs.SumY) / InArgs.Num) / Variance;
+	OutOffset = TargetY - OutSlope * TargetX;
+	return true;
+}
+
+FLinearRegressionArgs ComputeLinearRegressionInputArgs(const TConstArrayView<FVector2d>& InBuffer)
+{
+	FLinearRegressionArgs Result;
+	
+	Result.Num = InBuffer.Num();
+	if (Result.Num <= 0)
+	{
+		return Result;
+	}
+	
+	for (const FVector2d& Sample : InBuffer)
+	{
+		Result.SumX += Sample.X;
+		Result.SumY += Sample.Y;
+		Result.SumOfSquaredXes += Sample.X * Sample.Y;
+		Result.SumOfSquaredXes += FMath::Square(Sample.X);
+	}
+	return Result;
+}
+}
diff --git a/Engine/Source/Runtime/TimeManagement/Private/Tests/LinearRegressionTests.spec.cpp b/Engine/Source/Runtime/TimeManagement/Private/Tests/LinearRegressionTests.spec.cpp
new file mode 100644
--- /dev/null
+++ b/Engine/Source/Runtime/TimeManagement/Private/Tests/LinearRegressionTests.spec.cpp
@@ -0,0 +1,62 @@
+﻿// Copyright Epic Games, Inc. All Rights Reserved.
+
+#include "Misc/AutomationTest.h"
+
+#include "Misc/LinearRegression.h"
+#include "Misc/ModuloCircularBuffer.h"
+
+#if WITH_DEV_AUTOMATION_TESTS
+namespace UE::TimeManagement
+{
+BEGIN_DEFINE_SPEC(FLinearRegressionSpec, "System.Core.Time.LinearRegression", EAutomationTestFlags::EditorContext | EAutomationTestFlags::EngineFilter)
+	TModuloCircularBuffer<FVector2d> SampleBuffer = TModuloCircularBuffer<FVector2d>(5);
+END_DEFINE_SPEC(FLinearRegressionSpec);
+
+void FLinearRegressionSpec::Define()
+{
+	BeforeEach([this]
+	{
+		SampleBuffer = TModuloCircularBuffer<FVector2d>(5);
+	});
+
+	It("Simple regression", [this]
+	{
+		SampleBuffer.Add(FVector2d(0.25, 1));
+		SampleBuffer.Add(FVector2d(0.5, 2));
+		SampleBuffer.Add(FVector2d(1.0, 4));
+		
+		FLinearFunction Function;
+		ComputeLinearRegressionSlopeAndOffset(
+			ComputeLinearRegressionInputArgs(SampleBuffer.AsUnorderedView()),
+			Function
+			);
+
+		TestEqual(TEXT("(-0.25, -1)"),	Function.Evaluate(-0.25),	-1.0);
+		TestEqual(TEXT("(0, 0)"),		Function.Evaluate(0.0),		0.0);
+		TestEqual(TEXT("(0.75, 3)"),	Function.Evaluate(0.75),		3.0);
+		TestEqual(TEXT("(1.25, 5)"),	Function.Evaluate(1.25),		5.0);
+	});
+
+	It("Full buffer", [this]
+	{
+		SampleBuffer.Add(FVector2d(0, 0));	// This will be replaced, intentionally does not lie on a line with the other elements
+		SampleBuffer.Add(FVector2d(1, 20));
+		SampleBuffer.Add(FVector2d(2, 30));
+		SampleBuffer.Add(FVector2d(3, 40));
+		SampleBuffer.Add(FVector2d(4, 50));
+		SampleBuffer.Add(FVector2d(5, 60));	// This should replace (0,0) and make the function y = 10 + 10x.
+		
+		FLinearFunction Function;
+		ComputeLinearRegressionSlopeAndOffset(
+			ComputeLinearRegressionInputArgs(SampleBuffer.AsUnorderedView()),
+			Function
+			);
+		
+		TestEqual(TEXT("(-1, 0)"),		Function.Evaluate(-1.0),	0.0);
+		TestEqual(TEXT("(0, 10)"),		Function.Evaluate(0.0),		10.0);
+		TestEqual(TEXT("(0.5, 15)"),	Function.Evaluate(0.5),		15.0);
+		TestEqual(TEXT("(6, 70)"),		Function.Evaluate(6),		70.0);
+	});
+}
+}
+#endif
\ No newline at end of file
diff --git a/Engine/Source/Runtime/TimeManagement/Private/Tests/TimeManagementTests.cpp b/Engine/Source/Runtime/TimeManagement/Private/Tests/TimeManagementTests.cpp
new file mode 100644
--- /dev/null
+++ b/Engine/Source/Runtime/TimeManagement/Private/Tests/TimeManagementTests.cpp
@@ -0,0 +1,628 @@
+// Copyright Epic Games, Inc. All Rights Reserved.
+
+#include "Misc/AutomationTest.h"
+
+#if WITH_DEV_AUTOMATION_TESTS
+
+#include "Misc/FrameNumber.h"
+#include "Misc/FrameTime.h"
+#include "Misc/FrameRate.h"
+#include "CommonFrameRates.h"
+
+namespace
+{
+	FFrameTime TestTimes[] = {
+		FFrameTime(-10, 0.00f), FFrameTime(-10, 0.1f), FFrameTime(-10, 0.32f), FFrameTime(-10, 0.64f), FFrameTime(-10, 0.99999994f),
+		FFrameTime(-6,  0.00f), FFrameTime(-6,  0.1f), FFrameTime(-6,  0.32f), FFrameTime(-6,  0.64f), FFrameTime(-6,  0.99999994f),
+		FFrameTime(-5,  0.00f), FFrameTime(-5,  0.1f), FFrameTime(-5,  0.32f), FFrameTime(-5,  0.64f), FFrameTime(-5,  0.99999994f),
+		FFrameTime(-3,  0.00f), FFrameTime(-3,  0.1f), FFrameTime(-3,  0.32f), FFrameTime(-3,  0.64f), FFrameTime(-3,  0.99999994f),
+		FFrameTime(0,   0.00f), FFrameTime(0,   0.1f), FFrameTime(0,   0.32f), FFrameTime(0,   0.64f), FFrameTime(0,   0.99999994f),
+		FFrameTime(3,   0.00f), FFrameTime(3,   0.1f), FFrameTime(3,   0.32f), FFrameTime(3,   0.64f), FFrameTime(3,   0.99999994f),
+		FFrameTime(5,   0.00f), FFrameTime(5,   0.1f), FFrameTime(5,   0.32f), FFrameTime(5,   0.64f), FFrameTime(5,   0.99999994f),
+		FFrameTime(6,   0.00f), FFrameTime(6,   0.1f), FFrameTime(6,   0.32f), FFrameTime(6,   0.64f), FFrameTime(6,   0.99999994f),
+		FFrameTime(10,  0.00f), FFrameTime(10,  0.1f), FFrameTime(10,  0.32f), FFrameTime(10,  0.64f), FFrameTime(10,  0.99999994f),
+	};
+
+	bool IsNearlyEqual(FFrameTime Actual, FFrameTime Expected)
+	{
+		return Actual.FrameNumber == Expected.FrameNumber && FMath::IsNearlyEqual(Actual.GetSubFrame(), Expected.GetSubFrame(), KINDA_SMALL_NUMBER);
+	}
+}
+
+IMPLEMENT_SIMPLE_AUTOMATION_TEST(FFrameTimeComparisonTest, "System.Core.Time.Comparison", EAutomationTestFlags_ApplicationContextMask | EAutomationTestFlags::SmokeFilter | EAutomationTestFlags::HighPriority)
+bool FFrameTimeComparisonTest::RunTest(const FString& Parameters)
+{
+	const int32 NumFrames = UE_ARRAY_COUNT(TestTimes);
+
+	for (int32 Index = 0; Index < NumFrames; ++Index)
+	{
+		FFrameTime RHS = TestTimes[Index];
+
+		// Test frame times less than than current
+		for (int32 OtherIndex = 0; OtherIndex < Index; ++OtherIndex)
+		{
+			FFrameTime LHS = TestTimes[OtherIndex];
+
+			ensureAlwaysMsgf(  LHS <  RHS,  TEXT("  %d (+%0.3f) <   %d (+%0.3f)"), LHS.GetFrame().Value, LHS.GetSubFrame(), RHS.GetFrame().Value, RHS.GetSubFrame());
+			ensureAlwaysMsgf(  LHS <= RHS,  TEXT("  %d (+%0.3f) <=  %d (+%0.3f)"), LHS.GetFrame().Value, LHS.GetSubFrame(), RHS.GetFrame().Value, RHS.GetSubFrame());
+			ensureAlwaysMsgf(!(LHS >  RHS), TEXT("!(%d (+%0.3f) >  %d (+%0.3f))"), LHS.GetFrame().Value, LHS.GetSubFrame(), RHS.GetFrame().Value, RHS.GetSubFrame());
+			ensureAlwaysMsgf(!(LHS >= RHS), TEXT("!(%d (+%0.3f) >= %d (+%0.3f))"), LHS.GetFrame().Value, LHS.GetSubFrame(), RHS.GetFrame().Value, RHS.GetSubFrame());
+			ensureAlwaysMsgf(  LHS != RHS,  TEXT("  %d (+%0.3f) != %d (+%0.3f)"),  LHS.GetFrame().Value, LHS.GetSubFrame(), RHS.GetFrame().Value, RHS.GetSubFrame());
+			ensureAlwaysMsgf(!(LHS == RHS), TEXT("!(%d (+%0.3f) == %d (+%0.3f))"), LHS.GetFrame().Value, LHS.GetSubFrame(), RHS.GetFrame().Value, RHS.GetSubFrame());
+		}
+
+		// Test comparing against self
+		{
+			FFrameTime LHS = TestTimes[Index];
+			ensureAlwaysMsgf(!(LHS <  RHS), TEXT("!(%d (+%0.3f) <   %d (+%0.3f))"),LHS.GetFrame().Value, LHS.GetSubFrame(), RHS.GetFrame().Value, RHS.GetSubFrame());
+			ensureAlwaysMsgf(  LHS <= RHS,  TEXT("  %d (+%0.3f) <=  %d (+%0.3f)"), LHS.GetFrame().Value, LHS.GetSubFrame(), RHS.GetFrame().Value, RHS.GetSubFrame());
+			ensureAlwaysMsgf(!(LHS >  RHS), TEXT("!(%d (+%0.3f) >  %d (+%0.3f))"), LHS.GetFrame().Value, LHS.GetSubFrame(), RHS.GetFrame().Value, RHS.GetSubFrame());
+			ensureAlwaysMsgf(  LHS >= RHS,  TEXT("  %d (+%0.3f) >= %d (+%0.3f)"),  LHS.GetFrame().Value, LHS.GetSubFrame(), RHS.GetFrame().Value, RHS.GetSubFrame());
+			ensureAlwaysMsgf(!(LHS != RHS), TEXT("!(%d (+%0.3f) != %d (+%0.3f))"), LHS.GetFrame().Value, LHS.GetSubFrame(), RHS.GetFrame().Value, RHS.GetSubFrame());
+			ensureAlwaysMsgf(  LHS == RHS,  TEXT("  %d (+%0.3f) == %d (+%0.3f)"),  LHS.GetFrame().Value, LHS.GetSubFrame(), RHS.GetFrame().Value, RHS.GetSubFrame());
+		}
+
+		// Test frame times greater than current
+		for (int32 OtherIndex = Index + 1; OtherIndex < NumFrames; ++OtherIndex)
+		{
+			FFrameTime LHS = TestTimes[OtherIndex];
+
+			ensureAlwaysMsgf(!(LHS <  RHS), TEXT("!(%d (+%0.3f) <   %d (+%0.3f))"), LHS.GetFrame().Value, LHS.GetSubFrame(), RHS.GetFrame().Value, RHS.GetSubFrame());
+			ensureAlwaysMsgf(!(LHS <= RHS), TEXT("!(%d (+%0.3f) <=  %d (+%0.3f))"), LHS.GetFrame().Value, LHS.GetSubFrame(), RHS.GetFrame().Value, RHS.GetSubFrame());
+			ensureAlwaysMsgf(  LHS >  RHS,  TEXT("  %d (+%0.3f) >  %d (+%0.3f)"),   LHS.GetFrame().Value, LHS.GetSubFrame(), RHS.GetFrame().Value, RHS.GetSubFrame());
+			ensureAlwaysMsgf(  LHS >= RHS,  TEXT("  %d (+%0.3f) >= %d (+%0.3f)"),   LHS.GetFrame().Value, LHS.GetSubFrame(), RHS.GetFrame().Value, RHS.GetSubFrame());
+			ensureAlwaysMsgf(  LHS != RHS,  TEXT("  %d (+%0.3f) != %d (+%0.3f)"),   LHS.GetFrame().Value, LHS.GetSubFrame(), RHS.GetFrame().Value, RHS.GetSubFrame());
+			ensureAlwaysMsgf(!(LHS == RHS), TEXT("!(%d (+%0.3f) == %d (+%0.3f))"),  LHS.GetFrame().Value, LHS.GetSubFrame(), RHS.GetFrame().Value, RHS.GetSubFrame());
+		}
+	}
+
+
+	return true;
+}
+
+IMPLEMENT_SIMPLE_AUTOMATION_TEST(FFrameTimeToSecondsTest, "System.Core.Time.ToSeconds", EAutomationTestFlags_ApplicationContextMask | EAutomationTestFlags::SmokeFilter | EAutomationTestFlags::HighPriority)
+bool FFrameTimeToSecondsTest::RunTest(const FString& Parameters)
+{
+	const int32 NumFrames = UE_ARRAY_COUNT(TestTimes);
+
+	FFrameRate TestRate = FCommonFrameRates::FPS_60();
+	double ExpectedSeconds[] = {
+		-0.166666666666667 , -0.165             , -0.161333333333333 , -0.156             , -0.150000001       ,
+		-0.1               , -0.0983333333333333, -0.0946666666666667, -0.0893333333333333, -0.0833333343333333,
+		-0.0833333333333333, -0.0816666666666667, -0.078             , -0.0726666666666667, -0.0666666676666667,
+		-0.05              , -0.0483333333333333, -0.0446666666666667, -0.0393333333333333, -0.0333333343333333,
+		 0.                ,  0.0016666666666666,  0.0053333333333333,  0.0106666666666667,  0.0166666656666667,
+		 0.05              ,  0.0516666666666667,  0.0553333333333333,  0.0606666666666667,  0.0666666656666667,
+		 0.0833333333333333,  0.085             ,  0.0886666666666667,  0.094             ,  0.099999999       ,
+		 0.1               ,  0.101666666666667 ,  0.105333333333333 ,  0.110666666666667 ,  0.116666665666667 ,
+		 0.166666666666667 ,  0.168333333333333 ,  0.172             ,  0.177333333333333 ,  0.183333332333333 ,
+	};
+
+	for (int32 Index = 0; Index < NumFrames; ++Index)
+	{
+		FFrameTime Time     = TestTimes[Index];
+
+		// We test as floats so they round to the same precision as the sub frame
+		float     Actual   = float(Time / TestRate);
+		float     Expected = float(ExpectedSeconds[Index]);
+
+		ensureAlwaysMsgf(FMath::IsNearlyEqual(Actual, Expected, 6e-6f), TEXT("%d (+%0.3f) @ 60fps: %.9f seconds (actual) == %.9f (expected) seconds"),  Time.GetFrame().Value, Time.GetSubFrame(), Actual, Expected);
+	}
+
+
+	return true;
+}
+
+IMPLEMENT_SIMPLE_AUTOMATION_TEST(FFrameTimeAdditionTest, "System.Core.Time.Addition", EAutomationTestFlags_ApplicationContextMask | EAutomationTestFlags::SmokeFilter | EAutomationTestFlags::HighPriority)
+bool FFrameTimeAdditionTest::RunTest(const FString& Parameters)
+{
+	const int32 NumFrames = UE_ARRAY_COUNT(TestTimes);
+
+	// Test adding a positive FrameTime with a small sub frame
+	{
+		FFrameTime TimeToAdd(10, 0.1f);
+		FFrameTime ExpectedTimes[] = {
+			FFrameTime(0,  0.1f), FFrameTime(0,  0.2f), FFrameTime(0,  0.42f), FFrameTime(0,  0.74f), FFrameTime(1,  0.099999905f),
+			FFrameTime(4,  0.1f), FFrameTime(4,  0.2f), FFrameTime(4,  0.42f), FFrameTime(4,  0.74f), FFrameTime(5,  0.099999905f),
+			FFrameTime(5,  0.1f), FFrameTime(5,  0.2f), FFrameTime(5,  0.42f), FFrameTime(5,  0.74f), FFrameTime(6,  0.099999905f),
+			FFrameTime(7,  0.1f), FFrameTime(7,  0.2f), FFrameTime(7,  0.42f), FFrameTime(7,  0.74f), FFrameTime(8,  0.099999905f),
+			FFrameTime(10, 0.1f), FFrameTime(10, 0.2f), FFrameTime(10, 0.42f), FFrameTime(10, 0.74f), FFrameTime(11, 0.099999905f),
+			FFrameTime(13, 0.1f), FFrameTime(13, 0.2f), FFrameTime(13, 0.42f), FFrameTime(13, 0.74f), FFrameTime(14, 0.099999905f),
+			FFrameTime(15, 0.1f), FFrameTime(15, 0.2f), FFrameTime(15, 0.42f), FFrameTime(15, 0.74f), FFrameTime(16, 0.099999905f),
+			FFrameTime(16, 0.1f), FFrameTime(16, 0.2f), FFrameTime(16, 0.42f), FFrameTime(16, 0.74f), FFrameTime(17, 0.099999905f),
+			FFrameTime(20, 0.1f), FFrameTime(20, 0.2f), FFrameTime(20, 0.42f), FFrameTime(20, 0.74f), FFrameTime(21, 0.099999905f),
+		};
+
+		for (int32 Index = 0; Index < NumFrames; ++Index)
+		{
+			FFrameTime Time     = TestTimes[Index];
+			FFrameTime Actual   = Time + TimeToAdd;
+			FFrameTime Expected = ExpectedTimes[Index];
+
+			ensureAlwaysMsgf(IsNearlyEqual(Actual, Expected), TEXT("%d+%.3f + 10.1: %d+%.3f (actual) == %d+%.3f (expected)"),
+				Time.GetFrame().Value,     Time.GetSubFrame(),
+				Actual.GetFrame().Value,   Actual.GetSubFrame(),
+				Expected.GetFrame().Value, Expected.GetSubFrame()
+			);
+		}
+	}
+
+	// Test adding a positive FrameTime with a large sub frame
+	{
+		FFrameTime TimeToAdd(10, 0.8f);
+		FFrameTime ExpectedTimes[] = {
+			FFrameTime(0,  0.8f),  FFrameTime(0,  0.900000036f), FFrameTime(1,  0.120000005f), FFrameTime(1,  0.440000057f), FFrameTime(1,  0.79999997f),
+			FFrameTime(4,  0.8f),  FFrameTime(4,  0.900000036f), FFrameTime(5,  0.120000005f), FFrameTime(5,  0.440000057f), FFrameTime(5,  0.79999997f),
+			FFrameTime(5,  0.8f),  FFrameTime(5,  0.900000036f), FFrameTime(6,  0.120000005f), FFrameTime(6,  0.440000057f), FFrameTime(6,  0.79999997f),
+			FFrameTime(7,  0.8f),  FFrameTime(7,  0.900000036f), FFrameTime(8,  0.120000005f), FFrameTime(8,  0.440000057f), FFrameTime(8,  0.79999997f),
+			FFrameTime(10, 0.8f),  FFrameTime(10, 0.900000036f), FFrameTime(11, 0.120000005f), FFrameTime(11, 0.440000057f), FFrameTime(11, 0.79999997f),
+			FFrameTime(13, 0.8f),  FFrameTime(13, 0.900000036f), FFrameTime(14, 0.120000005f), FFrameTime(14, 0.440000057f), FFrameTime(14, 0.79999997f),
+			FFrameTime(15, 0.8f),  FFrameTime(15, 0.900000036f), FFrameTime(16, 0.120000005f), FFrameTime(16, 0.440000057f), FFrameTime(16, 0.79999997f),
+			FFrameTime(16, 0.8f),  FFrameTime(16, 0.900000036f), FFrameTime(17, 0.120000005f), FFrameTime(17, 0.440000057f), FFrameTime(17, 0.79999997f),
+			FFrameTime(20, 0.8f),  FFrameTime(20, 0.900000036f), FFrameTime(21, 0.120000005f), FFrameTime(21, 0.440000057f), FFrameTime(21, 0.79999997f),
+		};
+
+		for (int32 Index = 0; Index < NumFrames; ++Index)
+		{
+			FFrameTime Time     = TestTimes[Index];
+			FFrameTime Actual   = Time + TimeToAdd;
+			FFrameTime Expected = ExpectedTimes[Index];
+
+			ensureAlwaysMsgf(IsNearlyEqual(Actual, Expected), TEXT("%d+%.3f + 10.8: %d+%.3f (actual) == %d+%.3f (expected)"),
+				Time.GetFrame().Value,     Time.GetSubFrame(),
+				Actual.GetFrame().Value,   Actual.GetSubFrame(),
+				Expected.GetFrame().Value, Expected.GetSubFrame()
+			);
+		}
+	}
+
+	// Test adding a negative FrameTime with a tiny sub frame
+	{
+		FFrameTime TimeToAdd(-13, 0.01f);
+		FFrameTime ExpectedTimes[] = {
+			FFrameTime(-23, 0.01f), FFrameTime(-23, 0.11f), FFrameTime(-23, 0.329999983f), FFrameTime(-23, 0.65f), FFrameTime(-22, 0.00999999046f),
+			FFrameTime(-19, 0.01f), FFrameTime(-19, 0.11f), FFrameTime(-19, 0.329999983f), FFrameTime(-19, 0.65f), FFrameTime(-18, 0.00999999046f),
+			FFrameTime(-18, 0.01f), FFrameTime(-18, 0.11f), FFrameTime(-18, 0.329999983f), FFrameTime(-18, 0.65f), FFrameTime(-17, 0.00999999046f),
+			FFrameTime(-16, 0.01f), FFrameTime(-16, 0.11f), FFrameTime(-16, 0.329999983f), FFrameTime(-16, 0.65f), FFrameTime(-15, 0.00999999046f),
+			FFrameTime(-13, 0.01f), FFrameTime(-13, 0.11f), FFrameTime(-13, 0.329999983f), FFrameTime(-13, 0.65f), FFrameTime(-12, 0.00999999046f),
+			FFrameTime(-10, 0.01f), FFrameTime(-10, 0.11f), FFrameTime(-10, 0.329999983f), FFrameTime(-10, 0.65f), FFrameTime(-9,  0.00999999046f),
+			FFrameTime(-8,  0.01f), FFrameTime(-8,  0.11f), FFrameTime(-8,  0.329999983f), FFrameTime(-8,  0.65f), FFrameTime(-7,  0.00999999046f),
+			FFrameTime(-7,  0.01f), FFrameTime(-7,  0.11f), FFrameTime(-7,  0.329999983f), FFrameTime(-7,  0.65f), FFrameTime(-6,  0.00999999046f),
+			FFrameTime(-3,  0.01f), FFrameTime(-3,  0.11f), FFrameTime(-3,  0.329999983f), FFrameTime(-3,  0.65f), FFrameTime(-2,  0.00999999046f),
+		};
+
+		for (int32 Index = 0; Index < NumFrames; ++Index)
+		{
+			FFrameTime Time     = TestTimes[Index];
+			FFrameTime Actual   = Time + TimeToAdd;
+			FFrameTime Expected = ExpectedTimes[Index];
+
+			ensureAlwaysMsgf(IsNearlyEqual(Actual, Expected), TEXT("%d+%.3f + -13.01: %d+%.3f (actual) == %d+%.3f (expected)"),
+				Time.GetFrame().Value,     Time.GetSubFrame(),
+				Actual.GetFrame().Value,   Actual.GetSubFrame(),
+				Expected.GetFrame().Value, Expected.GetSubFrame()
+			);
+		}
+	}
+
+	// Test adding a negative FrameTime with a large sub frame
+	{
+		FFrameTime TimeToAdd(-13, 0.9f);
+		FFrameTime ExpectedTimes[] = {
+			FFrameTime(-23, 0.9f), FFrameTime(-22, 0.f), FFrameTime(-22, 0.220000029f), FFrameTime(-22, 0.539999962f), FFrameTime(-22, 0.899999857f),
+			FFrameTime(-19, 0.9f), FFrameTime(-18, 0.f), FFrameTime(-18, 0.220000029f), FFrameTime(-18, 0.539999962f), FFrameTime(-18, 0.899999857f),
+			FFrameTime(-18, 0.9f), FFrameTime(-17, 0.f), FFrameTime(-17, 0.220000029f), FFrameTime(-17, 0.539999962f), FFrameTime(-17, 0.899999857f),
+			FFrameTime(-16, 0.9f), FFrameTime(-15, 0.f), FFrameTime(-15, 0.220000029f), FFrameTime(-15, 0.539999962f), FFrameTime(-15, 0.899999857f),
+			FFrameTime(-13, 0.9f), FFrameTime(-12, 0.f), FFrameTime(-12, 0.220000029f), FFrameTime(-12, 0.539999962f), FFrameTime(-12, 0.899999857f),
+			FFrameTime(-10, 0.9f), FFrameTime(-9,  0.f), FFrameTime(-9,  0.220000029f), FFrameTime(-9,  0.539999962f), FFrameTime(-9,  0.899999857f),
+			FFrameTime(-8,  0.9f), FFrameTime(-7,  0.f), FFrameTime(-7,  0.220000029f), FFrameTime(-7,  0.539999962f), FFrameTime(-7,  0.899999857f),
+			FFrameTime(-7,  0.9f), FFrameTime(-6,  0.f), FFrameTime(-6,  0.220000029f), FFrameTime(-6,  0.539999962f), FFrameTime(-6,  0.899999857f),
+			FFrameTime(-3,  0.9f), FFrameTime(-2,  0.f), FFrameTime(-2,  0.220000029f), FFrameTime(-2,  0.539999962f), FFrameTime(-2,  0.899999857f), 
+		};
+
+		for (int32 Index = 0; Index < NumFrames; ++Index)
+		{
+			FFrameTime Time     = TestTimes[Index];
+			FFrameTime Actual   = Time + TimeToAdd;
+			FFrameTime Expected = ExpectedTimes[Index];
+
+			ensureAlwaysMsgf(IsNearlyEqual(Actual, Expected), TEXT("%d+%.3f + -13.9: %d+%.3f (actual) == %d+%.3f (expected)"),
+				Time.GetFrame().Value,     Time.GetSubFrame(),
+				Actual.GetFrame().Value,   Actual.GetSubFrame(),
+				Expected.GetFrame().Value, Expected.GetSubFrame()
+			);
+		}
+	}
+	return true;
+}
+
+IMPLEMENT_SIMPLE_AUTOMATION_TEST(FFrameTimeSubtractionTest, "System.Core.Time.Subtraction", EAutomationTestFlags_ApplicationContextMask | EAutomationTestFlags::SmokeFilter | EAutomationTestFlags::HighPriority)
+bool FFrameTimeSubtractionTest::RunTest(const FString& Parameters)
+{
+	const int32 NumFrames = UE_ARRAY_COUNT(TestTimes);
+
+	// Test subtracting a positive FrameTime with a small sub frame
+	{
+		FFrameTime TimeToSubtract(10, 0.1f);
+		FFrameTime ExpectedTimes[] = {
+			FFrameTime(-21, 0.9f), FFrameTime(-20, 0.f), FFrameTime(-20, 0.22f), FFrameTime(-20, 0.539999962f), FFrameTime(-20, 0.899999917f),
+			FFrameTime(-17, 0.9f), FFrameTime(-16, 0.f), FFrameTime(-16, 0.22f), FFrameTime(-16, 0.539999962f), FFrameTime(-16, 0.899999917f),
+			FFrameTime(-16, 0.9f), FFrameTime(-15, 0.f), FFrameTime(-15, 0.22f), FFrameTime(-15, 0.539999962f), FFrameTime(-15, 0.899999917f),
+			FFrameTime(-14, 0.9f), FFrameTime(-13, 0.f), FFrameTime(-13, 0.22f), FFrameTime(-13, 0.539999962f), FFrameTime(-13, 0.899999917f),
+			FFrameTime(-11, 0.9f), FFrameTime(-10, 0.f), FFrameTime(-10, 0.22f), FFrameTime(-10, 0.539999962f), FFrameTime(-10, 0.899999917f),
+			FFrameTime(-8,  0.9f),  FFrameTime(-7, 0.f), FFrameTime(-7,  0.22f), FFrameTime(-7,  0.539999962f), FFrameTime(-7,  0.899999917f),
+			FFrameTime(-6,  0.9f),  FFrameTime(-5, 0.f), FFrameTime(-5,  0.22f), FFrameTime(-5,  0.539999962f), FFrameTime(-5,  0.899999917f),
+			FFrameTime(-5,  0.9f),  FFrameTime(-4, 0.f), FFrameTime(-4,  0.22f), FFrameTime(-4,  0.539999962f), FFrameTime(-4,  0.899999917f),
+			FFrameTime(-1,  0.9f),  FFrameTime(0,  0.f), FFrameTime(0,   0.22f), FFrameTime(0,   0.539999962f), FFrameTime(0,   0.899999917f),
+		};
+
+		for (int32 Index = 0; Index < NumFrames; ++Index)
+		{
+			FFrameTime Time     = TestTimes[Index];
+			FFrameTime Actual   = Time - TimeToSubtract;
+			FFrameTime Expected = ExpectedTimes[Index];
+
+			ensureAlwaysMsgf(IsNearlyEqual(Actual, Expected), TEXT("%d+%.3f - 10.1: %d+%.3f (actual) == %d+%0.3f (expected)"),
+				Time.GetFrame().Value,     Time.GetSubFrame(),
+				Actual.GetFrame().Value,   Actual.GetSubFrame(),
+				Expected.GetFrame().Value, Expected.GetSubFrame()
+			);
+		}
+	}
+
+	// Test adding a positive FrameTime with a large sub frame
+	{
+		FFrameTime TimeToSubtract(10, 0.8f);
+		FFrameTime ExpectedTimes[] = {
+			FFrameTime(-21, 0.199999988f), FFrameTime(-21, 0.3f), FFrameTime(-21, 0.52f), FFrameTime(-21, 0.84f), FFrameTime(-20, 0.19999993f),
+			FFrameTime(-17, 0.199999988f), FFrameTime(-17, 0.3f), FFrameTime(-17, 0.52f), FFrameTime(-17, 0.84f), FFrameTime(-16, 0.19999993f),
+			FFrameTime(-16, 0.199999988f), FFrameTime(-16, 0.3f), FFrameTime(-16, 0.52f), FFrameTime(-16, 0.84f), FFrameTime(-15, 0.19999993f),
+			FFrameTime(-14, 0.199999988f), FFrameTime(-14, 0.3f), FFrameTime(-14, 0.52f), FFrameTime(-14, 0.84f), FFrameTime(-13, 0.19999993f),
+			FFrameTime(-11, 0.199999988f), FFrameTime(-11, 0.3f), FFrameTime(-11, 0.52f), FFrameTime(-11, 0.84f), FFrameTime(-10, 0.19999993f), 
+			FFrameTime(-8,  0.199999988f), FFrameTime(-8,  0.3f), FFrameTime(-8,  0.52f), FFrameTime(-8,  0.84f), FFrameTime(-7,  0.19999993f),
+			FFrameTime(-6,  0.199999988f), FFrameTime(-6,  0.3f), FFrameTime(-6,  0.52f), FFrameTime(-6,  0.84f), FFrameTime(-5,  0.19999993f),
+			FFrameTime(-5,  0.199999988f), FFrameTime(-5,  0.3f), FFrameTime(-5,  0.52f), FFrameTime(-5,  0.84f), FFrameTime(-4,  0.19999993f),
+			FFrameTime(-1,  0.199999988f), FFrameTime(-1,  0.3f), FFrameTime(-1,  0.52f), FFrameTime(-1,  0.84f), FFrameTime(0,   0.19999993f),
+		};
+
+		for (int32 Index = 0; Index < NumFrames; ++Index)
+		{
+			FFrameTime Time     = TestTimes[Index];
+			FFrameTime Actual   = Time - TimeToSubtract;
+			FFrameTime Expected = ExpectedTimes[Index];
+
+			ensureAlwaysMsgf(IsNearlyEqual(Actual, Expected), TEXT("%d+%.3f - 10.8: %d+%.3f (actual) == %d+%0.3f (expected)"),
+				Time.GetFrame().Value,     Time.GetSubFrame(),
+				Actual.GetFrame().Value,   Actual.GetSubFrame(),
+				Expected.GetFrame().Value, Expected.GetSubFrame()
+			);
+		}
+	}
+
+	// Test adding a negative FrameTime with a tiny sub frame
+	{
+		FFrameTime TimeToSubtract(-13, 0.01f);
+		FFrameTime ExpectedTimes[] = {
+			FFrameTime(2,  0.99f), FFrameTime(3,  0.09f), FFrameTime(3,  0.31f), FFrameTime(3,  0.63f), FFrameTime(3,  0.98999997f),
+			FFrameTime(6,  0.99f), FFrameTime(7,  0.09f), FFrameTime(7,  0.31f), FFrameTime(7,  0.63f), FFrameTime(7,  0.98999997f),
+			FFrameTime(7,  0.99f), FFrameTime(8,  0.09f), FFrameTime(8,  0.31f), FFrameTime(8,  0.63f), FFrameTime(8,  0.98999997f),
+			FFrameTime(9,  0.99f), FFrameTime(10, 0.09f), FFrameTime(10, 0.31f), FFrameTime(10, 0.63f), FFrameTime(10, 0.98999997f),
+			FFrameTime(12, 0.99f), FFrameTime(13, 0.09f), FFrameTime(13, 0.31f), FFrameTime(13, 0.63f), FFrameTime(13, 0.98999997f),
+			FFrameTime(15, 0.99f), FFrameTime(16, 0.09f), FFrameTime(16, 0.31f), FFrameTime(16, 0.63f), FFrameTime(16, 0.98999997f),
+			FFrameTime(17, 0.99f), FFrameTime(18, 0.09f), FFrameTime(18, 0.31f), FFrameTime(18, 0.63f), FFrameTime(18, 0.98999997f),
+			FFrameTime(18, 0.99f), FFrameTime(19, 0.09f), FFrameTime(19, 0.31f), FFrameTime(19, 0.63f), FFrameTime(19, 0.98999997f),
+			FFrameTime(22, 0.99f), FFrameTime(23, 0.09f), FFrameTime(23, 0.31f), FFrameTime(23, 0.63f), FFrameTime(23, 0.98999997f),
+		};
+
+		for (int32 Index = 0; Index < NumFrames; ++Index)
+		{
+			FFrameTime Time     = TestTimes[Index];
+			FFrameTime Actual   = Time - TimeToSubtract;
+			FFrameTime Expected = ExpectedTimes[Index];
+
+			ensureAlwaysMsgf(IsNearlyEqual(Actual, Expected), TEXT("%d+%.3f - -13.01: %d+%.3f (actual) == %d+%0.3f (expected)"),
+				Time.GetFrame().Value,     Time.GetSubFrame(),
+				Actual.GetFrame().Value,   Actual.GetSubFrame(),
+				Expected.GetFrame().Value, Expected.GetSubFrame()
+			);
+		}
+	}
+
+	// Test adding a negative FrameTime with a large sub frame
+	{
+		FFrameTime TimeToSubtract(-13, 0.9f);
+		FFrameTime ExpectedTimes[] = {
+			FFrameTime(2,  0.100000024f), FFrameTime(2,  0.200000048f), FFrameTime(2,  0.420000017f), FFrameTime(2,  0.74f), FFrameTime(3,  0.0999999642f),
+			FFrameTime(6,  0.100000024f), FFrameTime(6,  0.200000048f), FFrameTime(6,  0.420000017f), FFrameTime(6,  0.74f), FFrameTime(7,  0.0999999642f),
+			FFrameTime(7,  0.100000024f), FFrameTime(7,  0.200000048f), FFrameTime(7,  0.420000017f), FFrameTime(7,  0.74f), FFrameTime(8,  0.0999999642f),
+			FFrameTime(9,  0.100000024f), FFrameTime(9,  0.200000048f), FFrameTime(9,  0.420000017f), FFrameTime(9,  0.74f), FFrameTime(10, 0.0999999642f),
+			FFrameTime(12, 0.100000024f), FFrameTime(12, 0.200000048f), FFrameTime(12, 0.420000017f), FFrameTime(12, 0.74f), FFrameTime(13, 0.0999999642f),
+			FFrameTime(15, 0.100000024f), FFrameTime(15, 0.200000048f), FFrameTime(15, 0.420000017f), FFrameTime(15, 0.74f), FFrameTime(16, 0.0999999642f),
+			FFrameTime(17, 0.100000024f), FFrameTime(17, 0.200000048f), FFrameTime(17, 0.420000017f), FFrameTime(17, 0.74f), FFrameTime(18, 0.0999999642f),
+			FFrameTime(18, 0.100000024f), FFrameTime(18, 0.200000048f), FFrameTime(18, 0.420000017f), FFrameTime(18, 0.74f), FFrameTime(19, 0.0999999642f),
+			FFrameTime(22, 0.100000024f), FFrameTime(22, 0.200000048f), FFrameTime(22, 0.420000017f), FFrameTime(22, 0.74f), FFrameTime(23, 0.0999999642f),
+		};
+
+		for (int32 Index = 0; Index < NumFrames; ++Index)
+		{
+			FFrameTime Time     = TestTimes[Index];
+			FFrameTime Actual   = Time - TimeToSubtract;
+			FFrameTime Expected = ExpectedTimes[Index];
+
+			ensureAlwaysMsgf(IsNearlyEqual(Actual, Expected), TEXT("%d+%.3f - -13.9: %d+%.3f (actual) == %d+%0.3f (expected)"),
+				Time.GetFrame().Value,     Time.GetSubFrame(),
+				Actual.GetFrame().Value,   Actual.GetSubFrame(),
+				Expected.GetFrame().Value, Expected.GetSubFrame()
+			);
+		}
+	}
+	return true;
+}
+
+
+IMPLEMENT_SIMPLE_AUTOMATION_TEST(FFrameTimeConversionTest, "System.Core.Time.Conversion", EAutomationTestFlags_ApplicationContextMask | EAutomationTestFlags::SmokeFilter | EAutomationTestFlags::HighPriority)
+bool FFrameTimeConversionTest::RunTest(const FString& Parameters)
+{
+	const int32 NumFrames = UE_ARRAY_COUNT(TestTimes);
+
+	{
+		FFrameRate SrcRate = FCommonFrameRates::FPS_60();
+		FFrameRate DstRate = FCommonFrameRates::FPS_30();
+
+		FFrameTime ExpectedTimes[] = {
+			FFrameTime(-5, 0.f),  FFrameTime(-5, 0.05f), FFrameTime(-5, 0.16f), FFrameTime(-5, 0.32f), FFrameTime(-5, 0.499999985f),
+			FFrameTime(-3, 0.f),  FFrameTime(-3, 0.05f), FFrameTime(-3, 0.16f), FFrameTime(-3, 0.32f), FFrameTime(-3, 0.499999985f),
+			FFrameTime(-3, 0.5f), FFrameTime(-3, 0.55f), FFrameTime(-3, 0.66f), FFrameTime(-3, 0.82f), FFrameTime(-3, FFrameTime::MaxSubframe),
+			FFrameTime(-2, 0.5f), FFrameTime(-2, 0.55f), FFrameTime(-2, 0.66f), FFrameTime(-2, 0.82f), FFrameTime(-2, FFrameTime::MaxSubframe),
+			FFrameTime(0,  0.f),  FFrameTime(0,  0.05f), FFrameTime(0,  0.16f), FFrameTime(0,  0.32f), FFrameTime(0,  0.499999985f),
+			FFrameTime(1,  0.5f), FFrameTime(1,  0.55f), FFrameTime(1,  0.66f), FFrameTime(1,  0.82f), FFrameTime(1,  FFrameTime::MaxSubframe),
+			FFrameTime(2,  0.5f), FFrameTime(2,  0.55f), FFrameTime(2,  0.66f), FFrameTime(2,  0.82f), FFrameTime(2,  FFrameTime::MaxSubframe),
+			FFrameTime(3,  0.f),  FFrameTime(3,  0.05f), FFrameTime(3,  0.16f), FFrameTime(3,  0.32f), FFrameTime(3,  0.499999985f),
+			FFrameTime(5,  0.f),  FFrameTime(5,  0.05f), FFrameTime(5,  0.16f), FFrameTime(5,  0.32f), FFrameTime(5,  0.499999985f),
+		};
+
+		for (int32 Index = 0; Index < NumFrames; ++Index)
+		{
+			FFrameTime Time     = TestTimes[Index];
+			FFrameTime Actual   = FFrameRate::TransformTime(Time, SrcRate, DstRate);
+			FFrameTime Expected = ExpectedTimes[Index];
+
+			ensureAlwaysMsgf(Actual.GetFrame() == Expected.GetFrame() && FMath::IsNearlyEqual(Actual.GetSubFrame(), Expected.GetSubFrame(), KINDA_SMALL_NUMBER),
+				TEXT("%d+%.3f 60fps -> 30fps: %d+%.3f (actual) == %d+%0.3f (expected)"),
+				Time.GetFrame().Value,     Time.GetSubFrame(),
+				Actual.GetFrame().Value,   Actual.GetSubFrame(),
+				Expected.GetFrame().Value, Expected.GetSubFrame()
+			);
+		}
+	}
+
+	{
+		FFrameRate SrcRate = FCommonFrameRates::FPS_60();
+		FFrameRate DstRate = FCommonFrameRates::NTSC_30();
+
+		FFrameTime ExpectedTimes[] = {
+			FFrameTime(-5, 0.004995004995f), FFrameTime(-5, 0.05494505495f),  FFrameTime(-5, 0.1648351648f),   FFrameTime(-5, 0.3246753247f),   FFrameTime(-5, 0.5044954895f),
+			FFrameTime(-3, 0.002997002997f), FFrameTime(-3, 0.05294705295f),  FFrameTime(-3, 0.1628371628f),   FFrameTime(-3, 0.3226773227f),   FFrameTime(-3, 0.5024974875f),
+			FFrameTime(-3, 0.5024975025f),   FFrameTime(-3, 0.5524475524f),   FFrameTime(-3, 0.6623376623f),   FFrameTime(-3, 0.8221778222f),   FFrameTime(-2, 0.001997987013f),
+			FFrameTime(-2, 0.5014985015f),   FFrameTime(-2, 0.5514485514f),   FFrameTime(-2, 0.6613386613f),   FFrameTime(-2, 0.8211788212f),   FFrameTime(-1, 0.000998986014f),
+			FFrameTime(0,  0.f),             FFrameTime(0,  0.04995004995f),  FFrameTime(0,  0.1598401598f),   FFrameTime(0,  0.3196803197f),   FFrameTime(0,  0.4995004845f),
+			FFrameTime(1,  0.4985014985f),   FFrameTime(1,  0.5484515485f),   FFrameTime(1,  0.6583416583f),   FFrameTime(1,  0.8181818182f),   FFrameTime(1,  0.998001983f),
+			FFrameTime(2,  0.4975024975f),   FFrameTime(2,  0.5474525475f),   FFrameTime(2,  0.6573426573f),   FFrameTime(2,  0.8171828172f),   FFrameTime(2,  0.997002982f),
+			FFrameTime(2,  0.997002997f),    FFrameTime(3,  0.04695304695f),  FFrameTime(3,  0.1568431568f),   FFrameTime(3,  0.3166833167f),   FFrameTime(3,  0.4965034815f),
+			FFrameTime(4,  0.995004995f),    FFrameTime(5,  0.04495504496f),  FFrameTime(5,  0.1548451548f),   FFrameTime(5,  0.3146853147f),   FFrameTime(5,  0.4945054795f),
+		};
+
+		for (int32 Index = 0; Index < NumFrames; ++Index)
+		{
+			FFrameTime Time     = TestTimes[Index];
+			FFrameTime Actual   = FFrameRate::TransformTime(Time, SrcRate, DstRate);
+			FFrameTime Expected = ExpectedTimes[Index];
+
+			ensureAlwaysMsgf(Actual.GetFrame() == Expected.GetFrame() && FMath::IsNearlyEqual(Actual.GetSubFrame(), Expected.GetSubFrame(), KINDA_SMALL_NUMBER),
+				TEXT("%d+%.3f 60fps -> 29.97fps: %d+%.3f (actual) == %d+%0.3f (expected)"),
+				Time.GetFrame().Value,     Time.GetSubFrame(),
+				Actual.GetFrame().Value,   Actual.GetSubFrame(),
+				Expected.GetFrame().Value, Expected.GetSubFrame()
+			);
+		}
+	}
+
+	{
+		FFrameRate SrcRate = FCommonFrameRates::FPS_60();
+		FFrameRate DstRate = FCommonFrameRates::FPS_60();
+
+		for (int32 Index = 0; Index < NumFrames; ++Index)
+		{
+			FFrameTime Time     = TestTimes[Index];
+			FFrameTime Actual   = FFrameRate::TransformTime(Time, SrcRate, DstRate);
+			FFrameTime Expected = TestTimes[Index];
+
+			ensureAlwaysMsgf(IsNearlyEqual(Actual, Expected), TEXT("%d+%.3f 60fps -> 60fps: %d+%.3f (actual) == %d+%0.3f (expected)"),
+				Time.GetFrame().Value,     Time.GetSubFrame(),
+				Actual.GetFrame().Value,   Actual.GetSubFrame(),
+				Expected.GetFrame().Value, Expected.GetSubFrame()
+			);
+		}
+	}
+
+	return true;
+}
+
+IMPLEMENT_SIMPLE_AUTOMATION_TEST(FFrameRateMultiplesTest, "System.Core.Time.FrameRateMultiples", EAutomationTestFlags_ApplicationContextMask | EAutomationTestFlags::SmokeFilter | EAutomationTestFlags::HighPriority)
+bool FFrameRateMultiplesTest::RunTest(const FString& Parameters)
+{
+	FFrameRate TestRates[] = {
+		FCommonFrameRates::FPS_12(),
+		FCommonFrameRates::FPS_15(),
+		FCommonFrameRates::FPS_24(),
+		FCommonFrameRates::FPS_25(),
+		FCommonFrameRates::FPS_30(),
+		FCommonFrameRates::FPS_48(),
+		FCommonFrameRates::FPS_50(),
+		FCommonFrameRates::FPS_60(),
+		FCommonFrameRates::FPS_100(),
+		FCommonFrameRates::FPS_120(),
+		FCommonFrameRates::FPS_240(),
+		FCommonFrameRates::NTSC_24(),
+		FCommonFrameRates::NTSC_30(),
+		FCommonFrameRates::NTSC_60(),
+		FFrameRate(24000,1)
+	};
+
+	const int32 NumRates = UE_ARRAY_COUNT(TestRates);
+
+	{
+		bool IsMultipleOf[] = {
+			true,	false,	true,	false,	false,	true,	false,	true,	false,	true,	true,	false,	false,	false,	true,
+			false,	true,	false,	false,	true,	false,	false,	true,	false,	true,	true,	false,	false,	false,	true,
+			false,	false,	true,	false,	false,	true,	false,	false,	false,	true,	true,	false,	false,	false,	true,
+			false,	false,	false,	true,	false,	false,	true,	false,	true,	false,	false,	false,	false,	false,	true,
+			false,	false,	false,	false,	true,	false,	false,	true,	false,	true,	true,	false,	false,	false,	true,
+			false,	false,	false,	false,	false,	true,	false,	false,	false,	false,	true,	false,	false,	false,	true,
+			false,	false,	false,	false,	false,	false,	true,	false,	true,	false,	false,	false,	false,	false,	true,
+			false,	false,	false,	false,	false,	false,	false,	true,	false,	true,	true,	false,	false,	false,	true,
+			false,	false,	false,	false,	false,	false,	false,	false,	true,	false,	false,	false,	false,	false,	true,
+			false,	false,	false,	false,	false,	false,	false,	false,	false,	true,	true,	false,	false,	false,	true,
+			false,	false,	false,	false,	false,	false,	false,	false,	false,	false,	true,	false,	false,	false,	true,
+			false,	false,	false,	false,	false,	false,	false,	false,	false,	false,	false,	true,	false,	false,	true,
+			false,	false,	false,	false,	false,	false,	false,	false,	false,	false,	false,	false,	true,	true,	false,
+			false,	false,	false,	false,	false,	false,	false,	false,	false,	false,	false,	false,	false,	true,	false,
+			false,	false,	false,	false,	false,	false,	false,	false,	false,	false,	false,	false,	false,	false,	true,
+		};
+
+
+		bool IsFactorOf[] = {
+			true,	false,	false,	false,	false,	false,	false,	false,	false,	false,	false,	false,	false,	false,	false,
+			false,	true,	false,	false,	false,	false,	false,	false,	false,	false,	false,	false,	false,	false,	false,
+			true,	false,	true,	false,	false,	false,	false,	false,	false,	false,	false,	false,	false,	false,	false,
+			false,	false,	false,	true,	false,	false,	false,	false,	false,	false,	false,	false,	false,	false,	false,
+			false,	true,	false,	false,	true,	false,	false,	false,	false,	false,	false,	false,	false,	false,	false,
+			true,	false,	true,	false,	false,	true,	false,	false,	false,	false,	false,	false,	false,	false,	false,
+			false,	false,	false,	true,	false,	false,	true,	false,	false,	false,	false,	false,	false,	false,	false,
+			true,	true,	false,	false,	true,	false,	false,	true,	false,	false,	false,	false,	false,	false,	false,
+			false,	false,	false,	true,	false,	false,	true,	false,	true,	false,	false,	false,	false,	false,	false,
+			true,	true,	true,	false,	true,	false,	false,	true,	false,	true,	false,	false,	false,	false,	false,
+			true,	true,	true,	false,	true,	true,	false,	true,	false,	true,	true,	false,	false,	false,	false,
+			false,	false,	false,	false,	false,	false,	false,	false,	false,	false,	false,	true,	false,	false,	false,
+			false,	false,	false,	false,	false,	false,	false,	false,	false,	false,	false,	false,	true,	false,	false,
+			false,	false,	false,	false,	false,	false,	false,	false,	false,	false,	false,	false,	true,	true,	false,
+			true,	true,	true,	true,	true,	true,	true,	true,	true,	true,	true,	true,	false,	false,	true,
+		};
+
+		for (int32 Index = 0; Index < NumRates; ++Index)
+		{
+			FFrameRate SrcRate = TestRates[Index];
+
+			for (int32 OtherIndex = 0; OtherIndex < NumRates; ++OtherIndex)
+			{
+				FFrameRate OtherRate = TestRates[OtherIndex];
+
+				int32 TestResult = Index*NumRates + OtherIndex;
+
+				if (IsMultipleOf[TestResult])
+				{
+					ensureAlwaysMsgf(SrcRate.IsMultipleOf(OtherRate),
+						TEXT("Expected %d/%d to be a multiple of %d/%d."),
+						SrcRate.Numerator, SrcRate.Denominator,
+						OtherRate.Numerator, OtherRate.Denominator
+					);
+				}
+				else
+				{
+					ensureAlwaysMsgf(!SrcRate.IsMultipleOf(OtherRate),
+						TEXT("Did not expect %d/%d be a multiple of %d/%d."),
+						SrcRate.Numerator, SrcRate.Denominator,
+						OtherRate.Numerator, OtherRate.Denominator
+					);
+				}
+
+				if (IsFactorOf[TestResult])
+				{
+					ensureAlwaysMsgf(SrcRate.IsFactorOf(OtherRate),
+						TEXT("Expected %d/%d to be a factor of %d/%d."),
+						SrcRate.Numerator, SrcRate.Denominator,
+						OtherRate.Numerator, OtherRate.Denominator
+					);
+				}
+				else
+				{
+					ensureAlwaysMsgf(!SrcRate.IsFactorOf(OtherRate),
+						TEXT("Did not expect %d/%d to be a factor of %d/%d."),
+						SrcRate.Numerator, SrcRate.Denominator,
+						OtherRate.Numerator, OtherRate.Denominator
+					);
+				}
+			}
+		}
+	}
+
+	return true;
+}
+
+IMPLEMENT_SIMPLE_AUTOMATION_TEST(FFrameTimeToSecondsConversionTest, "System.Core.Time.FrameConversion", EAutomationTestFlags_ApplicationContextMask | EAutomationTestFlags::EngineFilter | EAutomationTestFlags::HighPriority)
+bool FFrameTimeToSecondsConversionTest::RunTest(const FString& Parameters)
+{
+	FFrameRate TestRates[] = {
+		FCommonFrameRates::FPS_12(),
+		FCommonFrameRates::FPS_15(),
+		FCommonFrameRates::FPS_24(),
+		FCommonFrameRates::FPS_25(),
+		FCommonFrameRates::FPS_30(),
+		FCommonFrameRates::FPS_48(),
+		FCommonFrameRates::FPS_50(),
+		FCommonFrameRates::FPS_60(),
+		FCommonFrameRates::FPS_100(),
+		FCommonFrameRates::FPS_120(),
+		FCommonFrameRates::FPS_240(),
+		FCommonFrameRates::NTSC_24(),
+		FCommonFrameRates::NTSC_30(),
+		FCommonFrameRates::NTSC_60(),
+		FFrameRate(24000,1)
+	};
+	
+	const int32 NumberOfSecondsToTest = 60;
+	for (const FFrameRate& FrameRate : TestRates)
+	{
+		const int32 NumberOfFrames = FrameRate.AsFrameNumber(static_cast<double>(NumberOfSecondsToTest)).Value;
+		const double Interval = FrameRate.AsInterval();
+		
+		for (int32 FrameIndex = 0; FrameIndex < NumberOfFrames; ++FrameIndex)
+		{
+			// Get seconds representation according to frame-rate
+			const double FrameAsSeconds = FrameRate.AsSeconds(FrameIndex);
+			const double IntervalSeconds = static_cast<double>(FrameIndex) * Interval;
+
+			// Verify seconds value against interval * FrameIndex
+			ensureAlwaysMsgf(FMath::IsNearlyEqual(IntervalSeconds, FrameAsSeconds),
+						TEXT("Did not expect value %f ::AsSeconds for input %i (%f)."),
+						FrameAsSeconds, FrameIndex, IntervalSeconds
+					);
+			
+			// Convert seconds values to FrameTime/FrameNumber values
+			const FFrameTime FrameTime = FrameRate.AsFrameTime(FrameAsSeconds);
+			const FFrameNumber FrameNumber = FrameRate.AsFrameNumber(FrameAsSeconds);
+
+			// Verify frame numbers match original FrameIndex value
+			ensureAlwaysMsgf(FrameTime.GetFrame() == FrameIndex,
+						TEXT("Did not expect frame number %i ::AsFrameTime for input %i (%f)."),
+						FrameNumber.Value, FrameIndex, FrameAsSeconds
+					);
+			
+			ensureAlwaysMsgf(FrameNumber.Value == FrameIndex,
+						TEXT("Did not expect frame number %i ::AsFrameNumber for input %i (%f)."),
+						FrameNumber.Value, FrameIndex, FrameAsSeconds
+					);
+			
+			// Convert FrameTime back to seconds and verify against value used to generate it
+			const double ConvertedSeconds = FrameRate.AsSeconds(FrameTime);
+			ensureAlwaysMsgf(FMath::IsNearlyEqual(ConvertedSeconds, FrameAsSeconds),
+				TEXT("Did not expect value %f ::AsSeconds for input %i (%f) - %s."),
+				ConvertedSeconds, FrameIndex, FrameAsSeconds, *LexToString(FrameTime)
+			);
+		}
+	}
+
+	return true;
+}
+
+#endif // WITH_DEV_AUTOMATION_TESTS
\ No newline at end of file
diff --git a/Engine/Source/Runtime/TimeManagement/Private/TimeManagementTests.cpp b/Engine/Source/Runtime/TimeManagement/Private/TimeManagementTests.cpp
deleted file mode 100644
diff --git a/Engine/Source/Runtime/TimeManagement/Public/CatchupFixedRateCustomTimeStep.h b/Engine/Source/Runtime/TimeManagement/Public/CatchupFixedRateCustomTimeStep.h
--- a/Engine/Source/Runtime/TimeManagement/Public/CatchupFixedRateCustomTimeStep.h
+++ b/Engine/Source/Runtime/TimeManagement/Public/CatchupFixedRateCustomTimeStep.h
@@ -3,7 +3,7 @@
 #pragma once
 
 #include "FixedFrameRateCustomTimeStep.h"
-
+#include "Estimation/IClockedTimeStep.h"
 #include "CatchupFixedRateCustomTimeStep.generated.h"
 
 class UEngine;
@@ -19,11 +19,12 @@ class UObject;
  * 
  */
 UCLASS(Blueprintable, editinlinenew, meta = (DisplayName = "Catchup Fixed Rate"), MinimalAPI)
-class UCatchupFixedRateCustomTimeStep : public UFixedFrameRateCustomTimeStep
+class UCatchupFixedRateCustomTimeStep : public UFixedFrameRateCustomTimeStep, public IClockedTimeStep
 {
 	GENERATED_UCLASS_BODY()
 
 public:
+	
 	//~ Begin UFixedFrameRateCustomTimeStep interface
 	TIMEMANAGEMENT_API virtual bool Initialize(UEngine* InEngine) override;
 	TIMEMANAGEMENT_API virtual void Shutdown(UEngine* InEngine) override;
@@ -32,6 +33,10 @@ public:
 	TIMEMANAGEMENT_API virtual FFrameRate GetFixedFrameRate() const override;
 	//~ End UFixedFrameRateCustomTimeStep interface
 
+	//~ Begin IClockedTimeStep interface
+	TIMEMANAGEMENT_API virtual TOptional<double> GetUnderlyingClockTime_AnyThread() override;
+	//~ End IClockedTimeStep interface
+	
 private:
 
 	/** Blocks until platform time reaches the given time */
diff --git a/Engine/Source/Runtime/TimeManagement/Public/Estimation/IClockedTimeStep.h b/Engine/Source/Runtime/TimeManagement/Public/Estimation/IClockedTimeStep.h
new file mode 100644
--- /dev/null
+++ b/Engine/Source/Runtime/TimeManagement/Public/Estimation/IClockedTimeStep.h
@@ -0,0 +1,28 @@
+﻿// Copyright Epic Games, Inc. All Rights Reserved.
+
+#pragma once
+
+#include "UObject/Interface.h"
+#include "IClockedTimeStep.generated.h"
+
+UINTERFACE(meta = (CannotImplementInterfaceInBlueprint))
+class UClockedTimeStep : public UInterface
+{
+	GENERATED_BODY()
+};
+
+/**
+ * Intended to be implemented by UEngineCustomTimeStep implementations that have an underlying clock they synchronize to.
+ * This is used, e.g. by UTimecodeRegressionProvider to estimate the timecode when the engine hitches.
+ */
+class IClockedTimeStep
+{
+	GENERATED_BODY()
+public:
+	
+	/**
+	 * @return The latest / immediate value of the clock. The clock is what the implementation uses to determine the DeltaTime in UpdateTimeStep.
+	 * If the current time is not available (e.g. still setting up a remote protocol) then this returns unset.
+	 */
+	virtual TOptional<double> GetUnderlyingClockTime_AnyThread() = 0;
+};
\ No newline at end of file
diff --git a/Engine/Source/Runtime/TimeManagement/Public/Estimation/TimecodeEstimator.h b/Engine/Source/Runtime/TimeManagement/Public/Estimation/TimecodeEstimator.h
new file mode 100644
--- /dev/null
+++ b/Engine/Source/Runtime/TimeManagement/Public/Estimation/TimecodeEstimator.h
@@ -0,0 +1,118 @@
+﻿// Copyright Epic Games, Inc. All Rights Reserved.
+
+#pragma once
+
+#include "Misc/App.h"
+#include "Misc/CachedLinearRegressionSums.h"
+#include "Misc/FrameRate.h"
+#include "Misc/LinearFunction.h"
+
+#define UE_API TIMEMANAGEMENT_API
+
+class IClockedTimeStep;
+class UTimecodeProvider;
+struct FQualifiedFrameTime;
+
+namespace UE::TimeManagement::TimecodeEstimation
+{
+/**
+ * Estimates the current timecode based on an IClockedTimeStep implementation, which is designed to be an UEngineCustomTimeStep.
+ *
+ * The engine starts each frame by calling UEngineCustomTimeStep::UpdateTimeStep. Then, using UTimecodeProvider::FetchAndUpdate, the engine calls
+ * FApp::SetCurrentFrameTime with the result of UTimecodeProvider::GetQualifiedFrameTime. Workflow:
+ * - FTimecodeEstimator: FetchAndUpdate samples the time code and tags it using the underlying clock's actual time (platform time, PTP, etc.), which
+ * is retrieved using IClockedTimeStep.
+ * - FTimecodeEstimator::GetQualifiedFrameTime estimates the current frame's time code by using the FApp::CurrentTime for linear regression of the
+ * sampled time codes. For this to work, FApp::CurrentTime is expected to be accumulation of all past delta times the UCustomTimeStep step has issued,
+ * which is sometimes called "game time" or "simulation time".
+ *
+ * If coupled with a UEngineCustomTimeStep that implements a fixed engine step rate, we can effectively handle hitching game frames, i.e. when frames
+ * take longer than the frame rate dedicated by the time code provider. Some systems, like Live Link, are used for querying external data; for the
+ * look-up, we use the frame's time code of the frame. However, when a frame takes longer, the subsequent frame needs to use the timecode value that
+ * was intended for that frame. The previous engine behaviour was to use FPlatformTime::Seconds() to determine the timecode the frame should have,
+ * which can cause the subsequent frame to inherit frame hitches. 
+ * 
+ * Explaining the issue with an example (TC = timecode):
+ * - The external timecode device's frame is set to 24 FPS, i.e. the frame budget is 0.0416666667s
+ * - Frame n is annotated with TC = 00:09:15.004.
+ * - Frame n takes 0.2s to process.
+ * - While frame n was running, the timecode's frame actually increased by 5 frames to 00:09:15.009 (i.e. real time passed by 5 target frames worth).
+ * Behaviours:
+ * - Old:
+ *	- We used to use the current platform time to determine timecode. This makes sense because TC is actually linearly correlated with physical time.
+ *	- So frame n+1 would use 00:09:15.009. Passing this to Live Link would skip the 5 frames of past data, and we'd get jumps in evaluated data. So
+ *	the simulation would skip 5 frames of live link data.
+ * - New (you'd use UTimecodeRegressionProvider):
+ *	- We'll estimate the timecode using linear regression.
+ *	- While the actual platform time has moved by 0.2s, FApp::CurrentTime should have only elapsed by DeltaTime (to simplify assume DeltaTime = 0.0416s).
+ *	- We ASSUME that DeltaTime is in the same time unit as the clock used internally in the custom time step, which could be FPlatformTime, PTP,
+ *	Rivermax time, Genlock time, etc. Basically, see what IClockedTimeStep::GetUnderlyingClockTime_AnyThread returns.
+ *	So frame n+1 would now use 00:09:15.005, which corresponds to the data that was sent to Live Link by external devices.
+ *	- Above we assumed that DeltaTime moves forward by 0.0416s, but the time step can decide this.
+ *		- Keeping DeltaTime = 0.0416s may cause the engine to never catch up with the external world but ensures that every frame always processes the
+ *		data for each frame (good for Take Recording).
+ *		- Increasing DeltaTime will increase the game time faster, thus allowing the engine to catch up, but to also skip recorded frame data. It can
+ *		result in visual jumps (good for real-time applications where the engine should not fall behind too much).
+ */
+class FTimecodeEstimator
+{
+public:
+
+	/**
+	 * @param InNumSamples The number of samples to used for linear regression.
+	 * @param InTimecode The timecoder provider for which we estimate the current frame's time. Caller ensures this outlives the constructed FTimecodeEstimator.
+	 * @param InEngineCustomTimeStep The provider of the current clock time. Caller ensures this outlives the constructed FTimecodeEstimator.
+	 */
+	UE_API explicit FTimecodeEstimator(
+		SIZE_T InNumSamples,
+		UTimecodeProvider& InTimecode UE_LIFETIMEBOUND, IClockedTimeStep& InEngineCustomTimeStep UE_LIFETIMEBOUND
+		);
+
+	/** Samples the current timecode and associates it with the underlying clock value. */
+	UE_API void FetchAndUpdate();
+	
+	/** Estimates what the current frame time should be given FApp::CurrentTime's value. */
+	UE_API FQualifiedFrameTime EstimateFrameTime() const;
+	
+private:
+
+	/**
+	 * The clock time when we were initialized.
+	 * This is subtracted from IClockedTimeStep::GetUnderlyingClockTime_AnyThread when used.
+	 * 
+	 * Clock times are subtracted with this value before they are passed to LinearRegression.
+	 * This effectively makes all values relative to the start time.
+	 * For example, in the linear regression input time 0.0 -> 00:09:15.009, time 0.4 -> 00:09:15.014, etc.
+	 *
+	 * The reason for this is minimize double precision issues.
+	 * E.g. FPlatformTime::Seconds() adds 16777216.0 to the result, which when tested caused a lot of numerical instability for the linear regression.
+	 * Doubles are more accurate the closer you are at 0 so we want to measure as close to that as possible.
+	 */
+	TOptional<double> StartClockTime;
+
+	/** Provides the actual time code. */
+	UTimecodeProvider& TimecodeProvider;
+	/** Provides the current clock time. */
+	IClockedTimeStep& EngineCustomTimeStep;
+	
+	/** Linear function that is used for predicting timecode (Y, dependent variable) based on clock time (X, independent variable) */
+	FLinearFunction LinearRegressionFunction;
+	
+	using FClockTimecodeSample = FVector2d; // X = clock time, Y = timecode converted using FFrameRate::AsSeconds.
+	/**
+	 * Used for computing the TimecodeLinearRegression based on frame time (Y, dependent variable) based on clock time (X, independent variable).
+	 * - Clock time is already a double.
+	 * - The frame time is timecode, i.e. the format 00:09:15.009. To do linear math with it, we must convert it to a number.
+	 *  For this we use FFrameRate::AsSeconds.
+	 */
+	FCachedLinearRegressionSums ClockToTimecodeSamples;
+	
+	/**
+	 * The last frame rate reported by the time code provider. Used to convert timecode to a double for linear regression.
+	 * If the value changes, the linear regression sampling buffer needs to be cleared.
+	 */
+	FFrameRate LastFrameRate;
+};
+}
+
+#undef UE_API
\ No newline at end of file
diff --git a/Engine/Source/Runtime/TimeManagement/Public/Estimation/TimecodeRegressionProvider.h b/Engine/Source/Runtime/TimeManagement/Public/Estimation/TimecodeRegressionProvider.h
new file mode 100644
--- /dev/null
+++ b/Engine/Source/Runtime/TimeManagement/Public/Estimation/TimecodeRegressionProvider.h
@@ -0,0 +1,89 @@
+﻿// Copyright Epic Games, Inc. All Rights Reserved.
+
+#pragma once
+
+#include "Engine/TimecodeProvider.h"
+#include "Estimation/TimecodeEstimator.h"
+#include "TimecodeRegressionProvider.generated.h"
+
+class UEngineCustomTimeStep;
+
+/**
+ * Estimates the current timecode based on the clock of the engine's current UCustomTimeStep.
+ * @see FTimecodeEstimator for a detailed explanation.
+ *
+ * This class is intended to be used for wrapping a UTimecodeProvider that was previously set using UEngine::SetTimecodeProvider.
+ * We probably want the engine to always run with UTimecodeRegressionProvider since otherwise frame hitches cause the timecode to be evaluated incorrectly.
+ * However, making that kind code change is risky so for now, we make systems that require it opt-in.
+ * In the future, we can refactor UEngine to use FTimecodeEstimator directly and then we can remove UTimecodeRegressionProvider.
+ */
+UCLASS()
+class TIMEMANAGEMENT_API UTimecodeRegressionProvider : public UTimecodeProvider
+{
+	GENERATED_BODY()
+public:
+
+	/**
+	 * Determines the number of samples to use for linear regression.
+	 * You'd want to set this number in relation to your target frame rate.
+	 * 
+	 * Example: If targeting a fixed frame rate of 24 then setting this to 24 would cause linear regression to effectively use 1s worth frames.
+	 * If you were targeting 60 FPS and wanted 1s worth of frames, you'd set this value to 60.
+	 */
+	UPROPERTY(EditAnywhere, BlueprintReadWrite, Category = "Settings", meta = (ClampMin = 1, ClampMax = 65536, UIMin = 1, UIMax = 65536))
+	int32 NumSampledFrames = 60;
+
+	/** @return Whether InTimestep is supported for timecode regression (non-null & implements the IClockedTimeStep interface). */
+	static bool SupportsTimestep(UEngineCustomTimeStep* InTimestep);
+
+	/** Sets the actual time code provider. You must call this before registering with the engine, i.e. before Initialize() is called. */
+	void SetTimecodeImplementation(UTimecodeProvider* InTimecodeImpl);
+
+	//~ Begin UTimecodeProvider Interface
+	virtual bool FetchTimecode(FQualifiedFrameTime& OutFrameTime) override
+	{
+		return TimecodeImpl->FetchTimecode(OutFrameTime);
+	}
+	virtual void FetchAndUpdate() override;
+	virtual FQualifiedFrameTime GetQualifiedFrameTime() const override;
+	virtual ETimecodeProviderSynchronizationState GetSynchronizationState() const override;
+	virtual bool Initialize(UEngine* InEngine) override;
+	virtual void Shutdown(UEngine* InEngine) override;
+	virtual bool SupportsAutoDetected() const override
+	{
+		return TimecodeImpl->SupportsAutoDetected();
+	}
+	virtual void SetIsAutoDetected(bool bInIsAutoDetected) override
+	{
+		TimecodeImpl->SetIsAutoDetected(bInIsAutoDetected);
+	}
+	virtual bool IsAutoDetected() const override
+	{
+		return TimecodeImpl->IsAutoDetected();
+	}
+	//~ End UTimecodeProvider Interface
+
+private:
+
+	/** The UEngine that called Initialize() on us. Used to respond to the custom time step changing. */
+	UPROPERTY()
+	TObjectPtr<UEngine> OwningEngine;
+	/** The real timecode provider, which this class is wrapping. */
+	UPROPERTY()
+	TObjectPtr<UTimecodeProvider> TimecodeImpl;
+
+	/** Estimates the timecode based on FApp::CurrentTime. */
+	TOptional<UE::TimeManagement::TimecodeEstimation::FTimecodeEstimator> TimecodeEstimator;
+
+	/** @return Whether Initialize has been called. */
+	bool IsInitialized() const
+	{
+		return OwningEngine != nullptr;
+	}
+
+	/** When the time step changes, the underlying clock may have changed: thus we should regenerate the sampled data. */
+	void OnCustomTimeStepChanged() { ReinitRegressionBuffer(); }
+
+	/** Clears the buffer if there is no custom time step. */
+	void ReinitRegressionBuffer();
+};
diff --git a/Engine/Source/Runtime/TimeManagement/Public/Misc/CachedLinearRegressionSums.h b/Engine/Source/Runtime/TimeManagement/Public/Misc/CachedLinearRegressionSums.h
new file mode 100644
--- /dev/null
+++ b/Engine/Source/Runtime/TimeManagement/Public/Misc/CachedLinearRegressionSums.h
@@ -0,0 +1,40 @@
+﻿// Copyright Epic Games, Inc. All Rights Reserved.
+
+#pragma once
+
+#include "LinearRegression.h"
+#include "ModuloCircularBuffer.h"
+
+namespace UE::TimeManagement
+{
+/**
+ * Data structure for building a FLinearFunction based on linear regression.
+ * 
+ * Caches the sums required for linear regression: when an old sample is replaced, the old value is subtracted and the new value added.
+ * This effectively avoids the summing all elements whenever the a new sample is added.
+ */
+struct FCachedLinearRegressionSums
+{
+	/** Holds the samples based off of which the samples are created. */
+	TModuloCircularBuffer<FVector2d> Samples;
+
+	/** Caches the sums of the samples. Whenever a sample is replaced, the old value is subtracted and the new value added to the respective sums. */
+	FLinearRegressionArgs CachedSums;
+
+	explicit FCachedLinearRegressionSums(SIZE_T InNumSamples) : Samples(InNumSamples) {}
+
+	/** @return Whether no samples have been added. */
+	bool IsEmpty() const
+	{
+		return Samples.IsEmpty();
+	}
+};
+
+/** Adds a sample and updates the associated sums. */
+void AddSampleAndUpdateSums(const FVector2d& InSample, TModuloCircularBuffer<FVector2d>& InSamples, FLinearRegressionArgs& InCachedArguments);
+	
+inline void AddSampleAndUpdateSums(const FVector2d& InSample, FCachedLinearRegressionSums& InCachedSums)
+{
+	AddSampleAndUpdateSums(InSample, InCachedSums.Samples, InCachedSums.CachedSums);
+}
+}
diff --git a/Engine/Source/Runtime/TimeManagement/Public/Misc/LinearFunction.h b/Engine/Source/Runtime/TimeManagement/Public/Misc/LinearFunction.h
new file mode 100644
--- /dev/null
+++ b/Engine/Source/Runtime/TimeManagement/Public/Misc/LinearFunction.h
@@ -0,0 +1,21 @@
+﻿// Copyright Epic Games, Inc. All Rights Reserved.
+
+#pragma once
+
+namespace UE::TimeManagement
+{
+/** Represents a linear function: f(x) = a*x + b */
+struct FLinearFunction
+{
+	/** The a in "f(x) = a*x + b" */
+	double Slope = 1;
+	/** The b in "f(x) = a*x + b" */
+	double Offset = 0;
+
+	/** Computes the function value using X as input. */
+	double Evaluate(double X) const 
+	{ 
+		return X * Slope + Offset; 
+	}
+};
+}
diff --git a/Engine/Source/Runtime/TimeManagement/Public/Misc/LinearRegression.h b/Engine/Source/Runtime/TimeManagement/Public/Misc/LinearRegression.h
new file mode 100644
--- /dev/null
+++ b/Engine/Source/Runtime/TimeManagement/Public/Misc/LinearRegression.h
@@ -0,0 +1,50 @@
+﻿// Copyright Epic Games, Inc. All Rights Reserved.
+
+#pragma once
+
+#include "Containers/ContainersFwd.h"
+#include "Math/Vector2D.h"
+#include "Misc/LinearFunction.h"
+
+namespace UE::TimeManagement
+{
+/** The input arguments required to compute linear regression. */
+struct FLinearRegressionArgs
+{
+	/** The sum of all X */
+	double SumX = 0.0;
+	/** The sum of all Y */
+	double SumY = 0.0;
+	/** The sum of each X multiplied with its Y. */
+	double SumXxY = 0.0;
+	/** The sum of all squared X values */
+	double SumOfSquaredXes = 0.0;
+	/** The number of samples */
+	SIZE_T Num = 0;
+};
+	
+/**
+ * @return The arguments required to compute linear regression based off of InBuffer.
+ * FVector2d::X is the independent variable, and FVector2d::Y the associated dependent variable value for the X. 
+ */
+FLinearRegressionArgs ComputeLinearRegressionInputArgs(const TConstArrayView<FVector2d>& InBuffer);
+
+/**
+ * Computes the coefficients for a linear function by using linear regression.
+ * @param InArgs The arguments required to compute linear regression
+ * @param OutSlope The slope for the linear regression. The a in "f(x) = a*x + b".
+ * @param OutOffset The offset for the linear regression. The b in "f(x) = a*x + b".
+ * @return True if the function was defined. False if the linear regression function was not defined, i.e. was no correlation.
+ * @note Sets OutSlope and OutOffset to 0 when the function is not defined.
+ */
+bool ComputeLinearRegressionSlopeAndOffset(
+	const FLinearRegressionArgs& InArgs,
+	double& OutSlope, double& OutOffset
+	);
+	
+/** Util that puts the sets the values of the linear function. */
+inline bool ComputeLinearRegressionSlopeAndOffset(const FLinearRegressionArgs& InArgs, FLinearFunction& OutFunction)
+{
+	return ComputeLinearRegressionSlopeAndOffset(InArgs, OutFunction.Slope, OutFunction.Offset);
+}
+}
diff --git a/Engine/Source/Runtime/TimeManagement/Public/Misc/ModuloCircularBuffer.h b/Engine/Source/Runtime/TimeManagement/Public/Misc/ModuloCircularBuffer.h
new file mode 100644
--- /dev/null
+++ b/Engine/Source/Runtime/TimeManagement/Public/Misc/ModuloCircularBuffer.h
@@ -0,0 +1,75 @@
+﻿// Copyright Epic Games, Inc. All Rights Reserved.
+
+#pragma once
+
+#include "Containers/Array.h"
+#include "Containers/ContainerAllocationPolicies.h"
+
+namespace UE::TimeManagement
+{
+/**
+ * A circular buffer of any size (opposed to TCircularBuffer which requires power of 2).
+ * 
+ * The next index is computed using the modulo operation: NextIndexToInsert = (NextIndexToInsert + 1) % Data.Num().
+ * It's less performant than bitmasking, which TCircularBuffer uses, but does not require any specific element count.
+ */
+template<typename T, typename TAllocator = FDefaultAllocator>
+class TModuloCircularBuffer
+{
+public:
+	
+	explicit TModuloCircularBuffer(SIZE_T InNumSamples) { Data.SetNumUninitialized(InNumSamples); }
+
+	/** Adds an item to the buffer. If the buffer is full, the oldest item is replaced. */
+	void Add(const T& InItem)
+	{
+		Data[NextIndexToInsert] = InItem;
+		bIsBufferFull |= NextIndexToInsert == Data.Num() - 1;
+		NextIndexToInsert = (NextIndexToInsert + 1) % Data.Num();
+	}
+
+	/** @return Pointer to the item that is replaced with the next Add call. If no item is replaced, then */
+	const T* GetNextReplacedItem() const
+	{
+		return bIsBufferFull ? &Data[NextIndexToInsert] : nullptr;
+	}
+
+	/** @return The number of items that can be stored in this buffer. */
+	SIZE_T Capacity() const 
+	{
+		return Data.Num();
+	}
+	
+	/** @return The number of items in the buffer so far. */
+	SIZE_T Num() const 
+	{
+		return bIsBufferFull ? Capacity() : NextIndexToInsert; 
+	}
+	/** @return Whether nothing has ever been added. */
+	bool IsEmpty() const 
+	{ 
+		return !bIsBufferFull && NextIndexToInsert == 0;
+	}
+	
+	/** @return Whether the buffer is full, i.e. add an item will override an item. */
+	bool IsFull() const 
+	{ 
+		return bIsBufferFull; 
+	}
+
+	/** @return A view into the data. Unordered here it does not represent the order in which items were added. This is useful if you want to sum up all entries, etc. */
+	TConstArrayView<T> AsUnorderedView() const
+	{
+		return bIsBufferFull ? TConstArrayView<T>(Data.GetData(), Data.Num()) : TConstArrayView<T>(Data.GetData(), NextIndexToInsert);
+	}
+
+private:
+
+	TArray<T, TAllocator> Data;
+
+	/** When recalculating, tells us how to sum up the samples. If false, we'll only sum of the first NextIndexToInsert elements of Samples. Otherwise all. */
+	bool bIsBufferFull = false;
+	/** The next item will be added to this index. Once bIsBufferFull == true, this points to the oldest data sample.  */
+	SIZE_T NextIndexToInsert = 0;
+};
+}
