Time Management: Add unit test framework for timecode estimation.

Fixes bug with the first frame being estimated as 00:00:00:00.00.

#Synced-CL 50395143

diff --git a/Engine/Source/Runtime/TimeManagement/Private/Misc/LinearRegression.cpp b/Engine/Source/Runtime/TimeManagement/Private/Misc/LinearRegression.cpp
--- a/Engine/Source/Runtime/TimeManagement/Private/Misc/LinearRegression.cpp
+++ b/Engine/Source/Runtime/TimeManagement/Private/Misc/LinearRegression.cpp
@@ -14,7 +14,9 @@ bool ComputeLinearRegressionSlopeAndOffset(
 	if (FMath::IsNearlyZero(Variance))
 	{
 		OutSlope = 0.0;
-		OutOffset = 0.0;
+		// The slope is undefined. The set of all possible linear curves is infinite. However, all of them must pass through the same point. 
+		// Therefore, we'll choose the linear function parallel to the x-axis.
+		OutOffset = InArgs.Num == 0 ? 0 : InArgs.SumY / InArgs.Num;
 		return false;
 	}
 
diff --git a/Engine/Source/Runtime/TimeManagement/Private/Tests/TimecodeEstimation.spec.cpp b/Engine/Source/Runtime/TimeManagement/Private/Tests/TimecodeEstimation.spec.cpp
new file mode 100644
--- /dev/null
+++ b/Engine/Source/Runtime/TimeManagement/Private/Tests/TimecodeEstimation.spec.cpp
@@ -0,0 +1,65 @@
+﻿// Copyright Epic Games, Inc. All Rights Reserved.
+
+#include "Misc/AutomationTest.h"
+#include "Utils/TimecodeEnvironmentMockWithTest.h"
+
+#if WITH_DEV_AUTOMATION_TESTS
+namespace UE::TimeManagement
+{
+BEGIN_DEFINE_SPEC(FTimecodeEstimationSpec, "System.Core.Time.TimecodeEstimation", EAutomationTestFlags::EditorContext | EAutomationTestFlags::EngineFilter)
+	TOptional<FTimecodeEnvironmentMockWithTest> Environment;
+
+	const double DeltaTime30Fps = 1.0 / 30.0;
+	/** A realistic value that FPlatformTime::Seconds() could return. Mocks the time since startup. */
+	const double StartTime = 1000;
+END_DEFINE_SPEC(FTimecodeEstimationSpec);
+
+void FTimecodeEstimationSpec::Define()
+{
+	BeforeEach([this]
+	{
+		Environment.Emplace(*this, FTimecodeEnvironmentArgs{ .StartTime = StartTime });
+	});
+	AfterEach([this]
+	{
+		Environment.Reset();
+	});
+	
+	It("5 frames w/o hitch", [this]
+	{
+		Environment->SimulateTickAndTest(FTimecode(12, 15, 30, 1, 0.0), FTimecode(12, 15, 30, 1, 0.0), DeltaTime30Fps, DeltaTime30Fps);
+		Environment->SimulateTickAndTest(FTimecode(12, 15, 30, 2, 0.0), FTimecode(12, 15, 30, 2, 0.0), DeltaTime30Fps, DeltaTime30Fps);
+		Environment->SimulateTickAndTest(FTimecode(12, 15, 30, 3, 0.0), FTimecode(12, 15, 30, 3, 0.0), DeltaTime30Fps, DeltaTime30Fps);
+		Environment->SimulateTickAndTest(FTimecode(12, 15, 30, 4, 0.0), FTimecode(12, 15, 30, 4, 0.0), DeltaTime30Fps, DeltaTime30Fps);
+		Environment->SimulateTickAndTest(FTimecode(12, 15, 30, 5, 0.0), FTimecode(12, 15, 30, 5, 0.0), DeltaTime30Fps, DeltaTime30Fps);
+	});
+	
+	Describe("5 frames with hitch on frame 3", [this]
+	{
+		const auto HitchOnFrame3 = [this](int32 TotalFrameDurationInNumFrames)
+		{
+			Environment->SimulateTickAndTest(FTimecode(12, 15, 30, 1, 0.0), FTimecode(12, 15, 30, 1, 0.0), DeltaTime30Fps, DeltaTime30Fps);
+			Environment->SimulateTickAndTest(FTimecode(12, 15, 30, 2, 0.0), FTimecode(12, 15, 30, 2, 0.0), DeltaTime30Fps, DeltaTime30Fps);
+				
+			// On frame 3, we'll take longer in physical time. Simulated time advance by a single DeltaTime.
+			const double PhysicalTimeElapsed = DeltaTime30Fps * TotalFrameDurationInNumFrames;
+			const int32 PhysicalFrame = 2 + TotalFrameDurationInNumFrames;
+			Environment->SimulateTickAndTest(
+				FTimecode(12, 15, 30, PhysicalFrame, 0.0),	// Real-world timecode
+				FTimecode(12, 15, 30, 3, 0.0),				// What it should estimate
+				PhysicalTimeElapsed,						// Physical delta time. We hitched
+				DeltaTime30Fps								// Simulation delta Time
+				);
+			
+			Environment->SimulateTickAndTest(FTimecode(12, 15, 30, PhysicalFrame + 1, 0.0), FTimecode(12, 15, 30, 4, 0.0), DeltaTime30Fps, DeltaTime30Fps);
+			Environment->SimulateTickAndTest(FTimecode(12, 15, 30, PhysicalFrame + 2, 0.0), FTimecode(12, 15, 30, 5, 0.0), DeltaTime30Fps, DeltaTime30Fps);
+		};
+		
+		// No matter how long the hitch, the estimator should correctly correlate the game time with physical time. 
+		It("Hitch 2 frames worth", [this, HitchOnFrame3]{ HitchOnFrame3(2); });
+		It("Hitch 10 frames worth", [this, HitchOnFrame3]{ HitchOnFrame3(10); });
+		It("Hitch 20 frames worth", [this, HitchOnFrame3]{ HitchOnFrame3(20); });
+	});
+}
+}
+#endif
\ No newline at end of file
diff --git a/Engine/Source/Runtime/TimeManagement/Private/Tests/Utils/MockTimecodeProvider.cpp b/Engine/Source/Runtime/TimeManagement/Private/Tests/Utils/MockTimecodeProvider.cpp
new file mode 100644
--- /dev/null
+++ b/Engine/Source/Runtime/TimeManagement/Private/Tests/Utils/MockTimecodeProvider.cpp
@@ -0,0 +1 @@
+﻿#include "MockTimecodeProvider.h"
diff --git a/Engine/Source/Runtime/TimeManagement/Private/Tests/Utils/MockTimecodeProvider.h b/Engine/Source/Runtime/TimeManagement/Private/Tests/Utils/MockTimecodeProvider.h
new file mode 100644
--- /dev/null
+++ b/Engine/Source/Runtime/TimeManagement/Private/Tests/Utils/MockTimecodeProvider.h
@@ -0,0 +1,41 @@
+﻿// Copyright Epic Games, Inc. All Rights Reserved.
+
+#pragma once
+
+#include "Engine/TimecodeProvider.h"
+#include "Misc/Attribute.h"
+#include "MockTimecodeProvider.generated.h"
+
+/** Timecode implementation useful for unit tests. */
+UCLASS(BlueprintInternalUseOnly) // BlueprintInternalUseOnly prevents this from showing up in class picker UI, e.g. details panel.
+class UMockTimecodeProvider : public UTimecodeProvider
+{
+	GENERATED_BODY()
+public:
+	
+	/** Gets the frame rate. Usually, the value will not change, but it's technically possible. */
+	TAttribute<FFrameRate> TargetFrameRate = FFrameRate(30, 1);
+	
+	/** Gets the current timecode */
+	TAttribute<FTimecode> FetchTimecodeAttr;
+
+	virtual bool FetchTimecode(FQualifiedFrameTime& OutFrameTime) override
+	{
+		const bool bHasAttr = FetchTimecodeAttr.IsSet() || FetchTimecodeAttr.IsBound();
+		if (bHasAttr)
+		{
+			OutFrameTime = FQualifiedFrameTime(FetchTimecodeAttr.Get(), TargetFrameRate.Get());
+		}
+		return bHasAttr;
+	}
+	
+	virtual void FetchAndUpdate() override { FetchTimecode(LatestTimecode); }
+	virtual FQualifiedFrameTime GetQualifiedFrameTime() const override { return LatestTimecode; }
+	virtual ETimecodeProviderSynchronizationState GetSynchronizationState() const override { return ETimecodeProviderSynchronizationState::Synchronized; }
+	virtual bool Initialize(UEngine* InEngine) override { return true; }
+	virtual void Shutdown(UEngine* InEngine) override {}
+
+private:
+	
+	FQualifiedFrameTime LatestTimecode;
+};
diff --git a/Engine/Source/Runtime/TimeManagement/Private/Tests/Utils/TimecodeEnvironmentMock.cpp b/Engine/Source/Runtime/TimeManagement/Private/Tests/Utils/TimecodeEnvironmentMock.cpp
new file mode 100644
--- /dev/null
+++ b/Engine/Source/Runtime/TimeManagement/Private/Tests/Utils/TimecodeEnvironmentMock.cpp
@@ -0,0 +1,61 @@
+﻿// Copyright Epic Games, Inc. All Rights Reserved.
+
+#include "TimecodeEnvironmentMock.h"
+
+#include "MockTimecodeProvider.h"
+#include "UObject/Package.h"
+
+namespace UE::TimeManagement
+{
+/** Util to revert side effects of unit tests on global engine state. Probably unnecessary. */
+struct FGuardAppTime
+{
+	const double RestoreTime = FApp::GetCurrentTime();
+	
+	FGuardAppTime(double InNewTime) { FApp::SetCurrentTime(InNewTime);}
+	~FGuardAppTime() { FApp::SetCurrentTime(RestoreTime); }
+};
+
+FTimecodeEnvironmentMock::FTimecodeEnvironmentMock(const FTimecodeEnvironmentArgs& InArgs)
+	: MockTimecodeProvider(NewObject<UMockTimecodeProvider>(GetTransientPackage()))
+	, Estimator(InArgs.NumSamples, *MockTimecodeProvider, *this)
+	, CurrentClockTime(InArgs.StartTime)
+	, AccumulatedTestAppCurrentTime(InArgs.StartTime)
+{
+	MockTimecodeProvider->TargetFrameRate = InArgs.TimecodeProviderFrameRate;
+	MockTimecodeProvider->FetchTimecodeAttr = TAttribute<FTimecode>::CreateLambda([this]{ return CurrentTimecode; });
+}
+
+void FTimecodeEnvironmentMock::SimulateTick(const FTimecode& InReportedTimecode, double InDeltaClockTime, double InGameDeltaTime)
+{
+	check(InDeltaClockTime > 0.0);
+	check(InGameDeltaTime > 0.0);
+	// If this fires, you're test is malformed. If timecode frame rate is 30 FPS, then e.g. timecode values  30, 50, etc. are invalid. 
+	const int32 MaxAllowedFrames = MockTimecodeProvider->TargetFrameRate.Get().Numerator;
+	check(MaxAllowedFrames > InReportedTimecode.Frames); 
+	
+	CurrentTimecode = InReportedTimecode;
+	CurrentClockTime += InDeltaClockTime;
+	AccumulatedTestAppCurrentTime += InGameDeltaTime;
+	
+	// FTimecodeEstimator uses FApp::CurrentTime under the hood.
+	const FGuardAppTime GuardTime(AccumulatedTestAppCurrentTime);
+	Estimator.FetchAndUpdate();
+}
+
+FTimecode FTimecodeEnvironmentMock::EstimateFrameTime() const
+{
+	const FGuardAppTime GuardTime(AccumulatedTestAppCurrentTime); // FTimecodeEstimator uses FApp::CurrentTime under the hood.
+	return Estimator.EstimateFrameTime().ToTimecode();
+}
+
+void FTimecodeEnvironmentMock::AddReferencedObjects(FReferenceCollector& Collector)
+{
+	Collector.AddReferencedObject(MockTimecodeProvider);
+}
+
+TOptional<double> FTimecodeEnvironmentMock::GetUnderlyingClockTime_AnyThread()
+{
+	return CurrentClockTime;
+}
+}
diff --git a/Engine/Source/Runtime/TimeManagement/Private/Tests/Utils/TimecodeEnvironmentMock.h b/Engine/Source/Runtime/TimeManagement/Private/Tests/Utils/TimecodeEnvironmentMock.h
new file mode 100644
--- /dev/null
+++ b/Engine/Source/Runtime/TimeManagement/Private/Tests/Utils/TimecodeEnvironmentMock.h
@@ -0,0 +1,92 @@
+﻿// Copyright Epic Games, Inc. All Rights Reserved.
+
+#pragma once
+
+#include "Estimation/IClockedTimeStep.h"
+#include "Estimation/TimecodeEstimator.h"
+#include "UObject/GCObject.h"
+
+class UMockTimecodeProvider;
+struct FTimecode;
+
+namespace UE::TimeManagement
+{
+struct FTimecodeEnvironmentArgs
+{
+	/** 
+	 * A reasonable value of FPlatformTime::Seconds(). 
+	 * Conceptually, this corresponds to the clock time when UEngineCustomTimeStep::Initialize is called.
+	 */
+	double StartTime = 1000;
+	
+	/** 
+	 * The frame rate that the mock timecode provider reports. 
+	 * This can be different from the fixed engine frame rate (and real world have engine fixed frame rate != timecode box frame rate). 
+	 */
+	FFrameRate TimecodeProviderFrameRate = FFrameRate(30, 1);
+	
+	/** Number of samples for linear regression. */
+	SIZE_T NumSamples = 30; // Equivalent to 1s at 30 FPS.
+};
+
+/**
+ * Simulates a simple engine environment for unit testing FTimecodeEstimator.
+ * It allows you to simulate that the engine ticked and sampled a specified FTimecode for that frame.
+ */
+class FTimecodeEnvironmentMock 
+	: private IClockedTimeStep
+	, public FGCObject
+	, public FNoncopyable
+{
+public:
+	
+	explicit FTimecodeEnvironmentMock(const FTimecodeEnvironmentArgs& InArgs = FTimecodeEnvironmentArgs());
+	
+	/**  
+	 * Makes a FTimecodeEstimator::FetchAndUpdate call causing it to 'sample' the specified FTimecode from the UTimecodeProvider and associated it 
+	 * with the given clock time.
+	 * 
+	 * @param InReportedTimecode The timecode that the simulated UTimecodeProvider would return in a real world scenario.
+	 * @param InDeltaClockTime The physical clock time that has elapsed since the last frame. Unit is seconds.
+	 *	If first frame, conceptually this is the time since UEngineCustomTimeStep::Initialize. 
+	 * @param InDeltaGameTime The delta time that custom delta time decides to advance the simulation time by. Unit is seconds.
+	 *	This affects the value final FApp::CurrentTime, which conceptually holds the simulation time. 
+	 */
+	void SimulateTick(const FTimecode& InReportedTimecode, double InDeltaClockTime, double InDeltaGameTime);
+	
+	/** @return The timecode that the FTimecodeEstimator estimates. */
+	FTimecode EstimateFrameTime() const;
+
+	//~ Begin FGCObject Interface
+	virtual void AddReferencedObjects(FReferenceCollector& Collector) override;
+	virtual FString GetReferencerName() const override { return TEXT("FTimecodeEnvironmentMock"); }
+	//~ End FGCObject Interface
+
+private:
+	
+	/** Mocktimecode provider */
+	TObjectPtr<UMockTimecodeProvider> MockTimecodeProvider;
+	/** The estimator being tested. */
+	TimecodeEstimation::FTimecodeEstimator Estimator;
+	
+	/** The clock time for the latest frame that was simulated. */
+	double CurrentClockTime = 0.0; 
+	
+	/** The timecode for the latest frame that was simulated. */
+	FTimecode CurrentTimecode;
+	
+	/** 
+	 * This is the accumulation of all simulated frames delta times. 
+	 * 
+	 * Conceptually, it's the time of the UWorld and corresponds to the value that FApp::CurrentTime is expected to hold. 
+	 * Under normal circumstances, custom time steps advance FApp::CurrentTime.
+	 */
+	double AccumulatedTestAppCurrentTime;
+
+	//~ Begin IClockedTimeStep Interface
+	virtual TOptional<double> GetUnderlyingClockTime_AnyThread() override;
+	//~ End IClockedTimeStep Interface
+};
+}
+
+
diff --git a/Engine/Source/Runtime/TimeManagement/Private/Tests/Utils/TimecodeEnvironmentMockWithTest.cpp b/Engine/Source/Runtime/TimeManagement/Private/Tests/Utils/TimecodeEnvironmentMockWithTest.cpp
new file mode 100644
--- /dev/null
+++ b/Engine/Source/Runtime/TimeManagement/Private/Tests/Utils/TimecodeEnvironmentMockWithTest.cpp
@@ -0,0 +1,34 @@
+﻿// Copyright Epic Games, Inc. All Rights Reserved.
+
+#include "TimecodeEnvironmentMockWithTest.h"
+
+#include "Misc/AutomationTest.h"
+
+namespace UE::TimeManagement
+{
+FTimecodeEnvironmentMockWithTest::FTimecodeEnvironmentMockWithTest(FAutomationTestBase& InTest, const FTimecodeEnvironmentArgs& InArgs)
+	: FTimecodeEnvironmentMock(InArgs)
+	, TestInstance(InTest)
+{}
+
+void FTimecodeEnvironmentMockWithTest::SimulateTickAndTest(
+	const FTimecode& InReportedTimecode, const FTimecode& InExpectedEstimatedTimecode,
+	double InDeltaClockTime, double InGameDeltaTime
+	)
+{
+	SimulateTick(InReportedTimecode, InDeltaClockTime, InGameDeltaTime);
+	
+	const FTimecode Estimate = EstimateFrameTime();
+	const bool bAreEqual = Estimate == InExpectedEstimatedTimecode;
+	if (!bAreEqual)
+	{
+		const bool bForceSubframe = InExpectedEstimatedTimecode.Subframe != Estimate.Subframe;
+		const FString Message = FString::Printf(
+			TEXT("Expected %s, but estimate is %s."), 
+			*InExpectedEstimatedTimecode.ToString(false, bForceSubframe), 
+			*Estimate.ToString(false, bForceSubframe)
+			);
+		TestInstance.AddError(*Message);
+	}
+}
+}
diff --git a/Engine/Source/Runtime/TimeManagement/Private/Tests/Utils/TimecodeEnvironmentMockWithTest.h b/Engine/Source/Runtime/TimeManagement/Private/Tests/Utils/TimecodeEnvironmentMockWithTest.h
new file mode 100644
--- /dev/null
+++ b/Engine/Source/Runtime/TimeManagement/Private/Tests/Utils/TimecodeEnvironmentMockWithTest.h
@@ -0,0 +1,45 @@
+﻿// Copyright Epic Games, Inc. All Rights Reserved.
+
+#pragma once
+
+#include "TimecodeEnvironmentMock.h"
+
+class FAutomationTestBase;
+struct FTimecode;
+
+namespace UE::TimeManagement
+{
+/** Version of FTimecodeEnvironmentMock that provides convenience functions to immediately test expected values (makes tests shorter = easier to read). */
+class FTimecodeEnvironmentMockWithTest : public FTimecodeEnvironmentMock
+{
+public:
+	
+	explicit FTimecodeEnvironmentMockWithTest(
+		FAutomationTestBase& InTest UE_LIFETIMEBOUND, 
+		const FTimecodeEnvironmentArgs& InArgs = {}
+		);
+	
+	/**
+	 * Makes a FTimecodeEstimator::FetchAndUpdate call and immediately tests the result of FTimecodeEstimator::EstimateFrameTime against the expected
+	 * value.
+	 * 
+	 * @param InReportedTimecode The timecode that the simulated UTimecodeProvider would return in a real world scenario.
+	 * @param InExpectedEstimatedTimecode The timecode value that you expect the FTimecodeEstimator to estimate.
+	 * @param InDeltaClockTime The physical clock time that has elapsed since the last frame. Unit is seconds.
+	 *	If first frame, conceptually this is the time since UEngineCustomTimeStep::Initialize. 
+	 * @param InDeltaGameTime The delta time that custom delta time decides to advance the simulation time by. Unit is seconds.
+	 *	This affects the value final FApp::CurrentTime, which conceptually holds the simulation time. 
+	 */
+	void SimulateTickAndTest(
+		const FTimecode& InReportedTimecode, const FTimecode& InExpectedEstimatedTimecode, 
+		double InDeltaClockTime, double InDeltaGameTime
+		);
+	
+private:
+	
+	/** The instance to run tests on. */
+	FAutomationTestBase& TestInstance;
+};
+}
+
+
