Improve timecode estimation with frame-coherent rounding

This change refactors the timecode estimator to use frame-coherent rounding that tracks frame advancement between estimates, preventing the "staircase problem" when converting
from continuous linear regression values to discrete timecode frames.

Key changes:

1. Frame-coherent rounding algorithm: The new RoundFrameTime() method examines the previous frame to determine correct rounding:
- If the base frame has advanced by exactly one from the previous estimate, use that frame value
- If the base frame matches the previous estimate, apply standard rounding
- Detects and logs unexpected conditions (backwards movement, frame jumps > 1)
2. Refactored estimation to use caching:
- EstimateFrameTime() now returns a cached value computed during FetchAndUpdate()
- Added CalculateEstimatedFrameTime() private method that performs the calculation and updates the cache
- Tracks previous estimated frame time to enable frame-coherent rounding
3. Removed unclear estimation warnings: Eliminated Timecode.UnclearEstimationSubframeThreshold console variable and associated warnings about values near 0.5 subframe,
simplifying the implementation in favor of deterministic frame tracking.
4. Enhanced diagnostics: Added logging for edge cases:
- Error log when timecode moves backwards (unexpected condition)
- Warning log when timecode jumps more than one frame (indicates game thread performance issues)
5. Bug fix: Removed duplicate TimecodeProvider.FetchAndUpdate() call in debug logging path
6. Minor code cleanup: Whitespace corrections and added missing include for QualifiedFrameTime.h

#Synced-CL 50715748

diff --git a/Engine/Source/Runtime/TimeManagement/Private/Estimation/TimecodeEstimator.cpp b/Engine/Source/Runtime/TimeManagement/Private/Estimation/TimecodeEstimator.cpp
--- a/Engine/Source/Runtime/TimeManagement/Private/Estimation/TimecodeEstimator.cpp
+++ b/Engine/Source/Runtime/TimeManagement/Private/Estimation/TimecodeEstimator.cpp
@@ -7,6 +7,7 @@
 #include "Estimation/IClockedTimeStep.h"
 #include "HAL/IConsoleManager.h"
 #include "ITimeManagementModule.h"
+#include "Misc/QualifiedFrameTime.h"
 
 namespace UE::TimeManagement::TimecodeEstimation
 {
@@ -19,11 +20,6 @@ static TAutoConsoleVariable<bool> CVarLogEstimation(
 static TAutoConsoleVariable<bool> CVarLogTimecodeDifference(
 	TEXT("Timecode.LogTimecodeDifference"), false, TEXT("Logs the timecode difference between the timecode of the current underlying clock and what is estimated. This is useful for debugging.")
 );
-static TAutoConsoleVariable<float> CVarUnclearEstimationSubframeThreshold(
-	TEXT("Timecode.UnclearEstimationSubframeThreshold"), 0.1f,
-	TEXT("If Abs(timecode's subframe  - 0.5f) <= this value, a warning is logged. That is because an estimated value like 14:11:43:16.49 is not a "
-		"clear estimate... it is very close to both 14:11:43:16.00, or 14:11:43:17.00")
-);
 
 static void LogEstimatedTime(
 	double InRelativeTime, IClockedTimeStep& Clock, UTimecodeProvider& TimecodeProvider,
@@ -48,7 +44,6 @@ static void LogEstimatedTime(
 
 	if (CVarLogTimecodeDifference.GetValueOnAnyThread())
 	{
-		TimecodeProvider.FetchAndUpdate();
 		const FQualifiedFrameTime ActualFrameTime = TimecodeProvider.GetQualifiedFrameTime();
 		
 		FTimecode ActualTC = ActualFrameTime.ToTimecode();
@@ -114,7 +109,7 @@ TOptional<FFetchAndUpdateStats> FTimecodeEstimator::FetchAndUpdate()
 	const FFrameRate CurrentFrameRate = CurrentFrameTime.Rate;
 	
 	// In a true production environment, the frame rate of the timecode device should not really change on the fly, but we should handle it anyway.
-	if (CurrentFrameRate != LastFrameRate) 
+	if (CurrentFrameRate != LastFrameRate)
 	{
 		ClockToTimecodeSamples = FCachedLinearRegressionSums(ClockToTimecodeSamples.Samples.Capacity());
 		LastFrameRate = CurrentFrameRate;
@@ -136,44 +131,93 @@ TOptional<FFetchAndUpdateStats> FTimecodeEstimator::FetchAndUpdate()
 		*ClockTime, FApp::GetCurrentTime()
 		);
 
+	CalculateEstimatedFrameTime();
 	return FFetchAndUpdateStats{ CurrentFrameTime };
 }
 
+
 FQualifiedFrameTime FTimecodeEstimator::EstimateFrameTime() const
+{
+	if (CachedEstimatedFrameTime)
+	{
+		return *CachedEstimatedFrameTime; 
+	}
+	return TimecodeProvider.GetQualifiedFrameTime();
+}
+	
+void FTimecodeEstimator::CalculateEstimatedFrameTime() 
 {
 	if (!ClockToTimecodeSamples.IsEmpty()
 		&& ensureMsgf(StartClockTime, TEXT("Invariant: StartClockTime was supposed to have been set when the data was sampled!")))
 	{
+		PreviousEstimatedFrameTime = CachedEstimatedFrameTime;
+
 		const double RelativeTime = FApp::GetCurrentTime() - *StartClockTime;
-		const double EstimatedSeconds = LinearRegressionFunction.Evaluate(RelativeTime); 
+		const double EstimatedSeconds = LinearRegressionFunction.Evaluate(RelativeTime);
 		const FFrameTime UnroundedEstimatedTime = EstimatedSeconds * LastFrameRate;
-		
-		// Subframe close to 0.5f? It may round to the wrong frame. Flag it so UE developer can investigate it, e.g. 0.47f is quite close 0.5f.
-		// We'd expect it to very close to the full frame, e.g. 14:11:43:16.04, or 14:11:43:16.87.
-		// A value like 14:11:43:16.42 is much closer to 0.5 than we'd expect... it may indicate that the time step is not linearly correlated
-		// to the timecode provider used (e.g. not really a fixed frame rate or some kind of noise).
-		const bool bUnclearEstimation = FMath::Abs(UnroundedEstimatedTime.GetSubFrame() - 0.5f) <= CVarUnclearEstimationSubframeThreshold.GetValueOnAnyThread();
-		constexpr bool bForceSignDisplay = false, bDisplaySubframe = true;
-		UE_CLOG(bUnclearEstimation, LogTimeManagement, Warning,
-			TEXT("Time %f estimate of %s is very close to 0.5 subframe border... the resulting timecode's frame may be off by 1."),
-			RelativeTime, *FQualifiedFrameTime(UnroundedEstimatedTime, LastFrameRate).ToTimecode().ToString(bForceSignDisplay, bDisplaySubframe)
-			);
-		
-		// FApp::GetCurrentTime() is usually slightly behind, e.g. relative underlying clock may be 2.028365, but FApp::GetCurrentTime() is 2.028364.
-		// That can cause slight trailing (provider 14:11:43:17.00, est. 14:11:43:16.99), or leading (provider: 14:11:43:19.00, est. 14:11:43:19.01).
-		// We only care about full frames... we drop subframes. That's why slight leading is no problem as the frame number stays the same.
-		// Trailing *is* a problem though because the frame is one less.
-		// Note: this comment was written when UCatchupFixedRateCustomTimeStep was the only timestep that UTimecodeRegressionProvider was set up with:
-		// so this observation is untested with other time steps that may have been added later (5.7+). If you find that behaviour is different
-		// between time steps, then we need to adjust this strategy.
-		const FFrameTime RoundedEstimatedTime = UnroundedEstimatedTime.RoundToFrame();
+
+		const FTimecode UnroundedTimecode = FQualifiedFrameTime(UnroundedEstimatedTime, LastFrameRate).ToTimecode();
+		const FFrameTime RoundedEstimatedTime = RoundFrameOrResetEstimator(UnroundedEstimatedTime);
 		
 		LogEstimatedTime(RelativeTime, EngineCustomTimeStep, TimecodeProvider, RoundedEstimatedTime, UnroundedEstimatedTime, LastFrameRate);
-		return FQualifiedFrameTime(RoundedEstimatedTime, LastFrameRate);
+
+		CachedEstimatedFrameTime = FQualifiedFrameTime(RoundedEstimatedTime, LastFrameRate);
+		return;
 	}
-	
+
 	// This may cause jumps at the beginning, but it can be circumvented by warming up the engine, i.e. just let it run for a few frames
 	UE_LOG(LogTimeManagement, Log, TEXT("No data sampled, yet. This frame will fall back to actual timecode without estimation."));
-	return TimecodeProvider.GetQualifiedFrameTime();
+	CachedEstimatedFrameTime = TimecodeProvider.GetQualifiedFrameTime();
+}
+
+FFrameTime FTimecodeEstimator::RoundFrameOrResetEstimator(const FFrameTime& UnroundedEstimatedTime)
+{
+	if (!PreviousEstimatedFrameTime)
+	{
+		return UnroundedEstimatedTime.RoundToFrame();
+	}
+
+	const FFrameNumber BaseFrame = UnroundedEstimatedTime.GetFrame();
+	const FFrameNumber PreviousFrame = PreviousEstimatedFrameTime->Time.GetFrame();
+	
+	if (BaseFrame == PreviousFrame + 1)
+	{
+		// This is the next frame for our TC value. This will handle cases where we jump from 10:01:10.24.49 to 10:01:10.25.50 in a 30 FPS
+		// environment. We would also go to the next instead of jumping to 10:01:10.26. 
+		return FFrameTime(BaseFrame);
+	}
+	else if (BaseFrame == PreviousFrame)
+	{
+		// We effectively just want to round here.  2 subframes would be .5 in Unreal's version of subframe so we can just round to find
+		// the next value.
+		return UnroundedEstimatedTime.RoundToFrame();
+	}
+
+	// Handle other edge cases.
+	//
+	// The likely case is that we have "jumped" ahead likely due to game thread not running fast enough. If we have jumped
+	// backwards then that is an indication something wrong has happened and should be reported to the user.
+	// 
+	const FTimecode BaseTimecode = FQualifiedFrameTime(UnroundedEstimatedTime, LastFrameRate).ToTimecode();
+	const FTimecode PreviousTimecode = PreviousEstimatedFrameTime->ToTimecode();
+	if (BaseFrame < PreviousFrame)
+	{
+		UE_LOG(LogTimeManagement, Warning, TEXT("We have calculated new timecode value (%s) that is before our current timecode (%s). This could be because of timecode rollover. This will reset our estimator to the start state. "), *BaseTimecode.ToString(false, true), *PreviousTimecode.ToString(false, true));
+
+		Reset();
+		return UnroundedEstimatedTime.RoundToFrame();
+	}
+
+	UE_LOG(LogTimeManagement, Warning, TEXT("FTimecodeEstimator has calculated a new timecode (%s) more than one frame greater than previous (%s). Likely because your game thread cannot keep up with desired rated."), *BaseTimecode.ToString(false, true), *PreviousTimecode.ToString(false,true)); 
+	return FFrameTime(BaseFrame);
 }
+
+void FTimecodeEstimator::Reset()
+{
+	LinearRegressionFunction = FLinearFunction();
+	PreviousEstimatedFrameTime.Reset();
+	ClockToTimecodeSamples = FCachedLinearRegressionSums(ClockToTimecodeSamples.Samples.Capacity());
+	StartClockTime.Reset();
+}
+	
 }
diff --git a/Engine/Source/Runtime/TimeManagement/Public/Estimation/TimecodeEstimator.h b/Engine/Source/Runtime/TimeManagement/Public/Estimation/TimecodeEstimator.h
--- a/Engine/Source/Runtime/TimeManagement/Public/Estimation/TimecodeEstimator.h
+++ b/Engine/Source/Runtime/TimeManagement/Public/Estimation/TimecodeEstimator.h
@@ -60,6 +60,9 @@ struct FFetchAndUpdateStats
  *		data for each frame (good for Take Recording).
  *		- Increasing DeltaTime will increase the game time faster, thus allowing the engine to catch up, but to also skip recorded frame data. It can
  *		result in visual jumps (good for real-time applications where the engine should not fall behind too much).
+ *
+ *	Note: the estimator currently assumes, and this is enforced by the current usage in take recorder setup, that the TC rate is equal
+ *	to the sampling rate.  Otherwise, addressing staircasing issues may be wrong. JIRA: UE-364878
  */
 class FTimecodeEstimator
 {
@@ -83,11 +86,31 @@ public:
 	 */
 	UE_API TOptional<FFetchAndUpdateStats> FetchAndUpdate();
 	
-	/** Estimates what the current frame time should be given FApp::CurrentTime's value. */
+	/**
+	 * Gets the current (cached) estimated frame time. If this is called before FetchAndUpdate() call then no cached timecode will
+	 * be available and thus it will use the base timecode implementation to query the estimated frame time.
+	 */
 	UE_API FQualifiedFrameTime EstimateFrameTime() const;
-	
+
 private:
+	/**
+	 * Calculate the next qualified frame time based on the samples store in our regression estimator.  This is called on FetchAndUpdate()
+	 * and updates the internal cached frametime value.
+	 */
+	void CalculateEstimatedFrameTime();
+
+	/**
+	 * Uses the previous frame time to determine the correct way to round the computed regression value to the correct TC value. This
+	 * function may reset the regressor if a negative step was detected.  This could happen in TC rollover or bad data from the
+	 * underlying provider. 
+	 */
+	FFrameTime RoundFrameOrResetEstimator(const FFrameTime& UnroundedEstimatedTime);
 
+	/**
+	 * Reset the estimator back to the startup state.  This is necessary when rollover or negative regression estimate is detected. 
+	 */
+	void Reset();
+	
 	/**
 	 * The clock time when we were initialized.
 	 * This is subtracted from IClockedTimeStep::GetUnderlyingClockTime_AnyThread when used.
@@ -104,12 +127,22 @@ private:
 
 	/** Provides the actual time code. */
 	UTimecodeProvider& TimecodeProvider;
+
 	/** Provides the current clock time. */
 	IClockedTimeStep& EngineCustomTimeStep;
 	
 	/** Linear function that is used for predicting timecode (Y, dependent variable) based on clock time (X, independent variable) */
 	FLinearFunction LinearRegressionFunction;
-	
+
+	/** Currently cached estimated frame time. */
+	TOptional<FQualifiedFrameTime> CachedEstimatedFrameTime;
+
+    /**
+	 * Previously calculated estimated frame time used to analyze which boundary we are on for the "next" timecode.  It helps
+	 * resolve the staircase problem when sampling from a linear regression to a discrete timecode value.
+	 */
+	TOptional<FQualifiedFrameTime> PreviousEstimatedFrameTime;
+
 	using FClockTimecodeSample = FVector2d; // X = clock time, Y = timecode converted using FFrameRate::AsSeconds.
 	/**
 	 * Used for computing the TimecodeLinearRegression based on frame time (Y, dependent variable) based on clock time (X, independent variable).
@@ -127,4 +160,4 @@ private:
 };
 }
 
-#undef UE_API
\ No newline at end of file
+#undef UE_API
